- Add foreign id registry

diff --git a/Engine/Plugins/Online/Experimental/OnlineServicesEOS/Source/Private/OnlineIdEOS.cpp b/Engine/Plugins/Online/Experimental/OnlineServicesEOS/Source/Private/OnlineIdEOS.cpp
--- a/Engine/Plugins/Online/Experimental/OnlineServicesEOS/Source/Private/OnlineIdEOS.cpp
+++ b/Engine/Plugins/Online/Experimental/OnlineServicesEOS/Source/Private/OnlineIdEOS.cpp
@@ -226,7 +226,7 @@ TArray<uint8> FOnlineAccountIdRegistryEOS::ToReplicationData(const FOnlineAccoun
 	return ReplicationData;
 }
 
-FOnlineAccountIdHandle FOnlineAccountIdRegistryEOS::FromReplicationData(const TArrayView<uint8> ReplicationData)
+FOnlineAccountIdHandle FOnlineAccountIdRegistryEOS::FromReplicationData(const TArray<uint8> ReplicationData)
 {
 	const EEOSAccountIdElements Elements = EEOSAccountIdElements(ReplicationData[0]);
 	const int EasHexBufferLength = EnumHasAnyFlags(Elements, EEOSAccountIdElements::EAS) ? OnlineIdEOSHexBufferLength : 0;
diff --git a/Engine/Plugins/Online/Experimental/OnlineServicesEOS/Source/Private/OnlineIdEOS.h b/Engine/Plugins/Online/Experimental/OnlineServicesEOS/Source/Private/OnlineIdEOS.h
--- a/Engine/Plugins/Online/Experimental/OnlineServicesEOS/Source/Private/OnlineIdEOS.h
+++ b/Engine/Plugins/Online/Experimental/OnlineServicesEOS/Source/Private/OnlineIdEOS.h
@@ -40,7 +40,7 @@ public:
 	// Begin IOnlineAccountIdRegistry
 	virtual FString ToLogString(const FOnlineAccountIdHandle& Handle) const override;
 	virtual TArray<uint8> ToReplicationData(const FOnlineAccountIdHandle& Handle) const override;
-	virtual FOnlineAccountIdHandle FromReplicationData(const TArrayView<uint8> ReplicationString) override;
+	virtual FOnlineAccountIdHandle FromReplicationData(const TArray<uint8> ReplicationString) override;
 	// End IOnlineAccountIdRegistry
 
 	virtual ~FOnlineAccountIdRegistryEOS() = default;
diff --git a/Engine/Source/Runtime/CoreOnline/CoreOnline.Build.cs b/Engine/Source/Runtime/CoreOnline/CoreOnline.Build.cs
--- a/Engine/Source/Runtime/CoreOnline/CoreOnline.Build.cs
+++ b/Engine/Source/Runtime/CoreOnline/CoreOnline.Build.cs
@@ -13,5 +13,7 @@ public class CoreOnline : ModuleRules
 				"CoreUObject"
 			}
 		);
+
+		PrivateDefinitions.Add("COREONLINE_PACKAGE=1");
 	}
 }
diff --git a/Engine/Source/Runtime/CoreOnline/Private/CoreOnlinePrivate.h b/Engine/Source/Runtime/CoreOnline/Private/CoreOnlinePrivate.h
deleted file mode 100644
diff --git a/Engine/Source/Runtime/CoreOnline/Private/Online/CoreOnline.cpp b/Engine/Source/Runtime/CoreOnline/Private/Online/CoreOnline.cpp
--- a/Engine/Source/Runtime/CoreOnline/Private/Online/CoreOnline.cpp
+++ b/Engine/Source/Runtime/CoreOnline/Private/Online/CoreOnline.cpp
@@ -3,6 +3,7 @@
 #include "Online/CoreOnline.h"
 
 #include "Misc/LazySingleton.h"
+#include "Online/CoreOnlinePrivate.h"
 
 namespace UE::Online {
 
@@ -11,6 +12,15 @@ FOnlineIdRegistryRegistry& FOnlineIdRegistryRegistry::Get()
 	return TLazySingleton<FOnlineIdRegistryRegistry>::Get();
 }
 
+FOnlineIdRegistryRegistry::FOnlineIdRegistryRegistry()
+	: ForeignAccountIdRegistry(MakeUnique<FOnlineForeignAccountIdRegistry>())
+{
+}
+
+FOnlineIdRegistryRegistry::~FOnlineIdRegistryRegistry()
+{
+}
+
 void FOnlineIdRegistryRegistry::TearDown()
 {
 	return TLazySingleton<FOnlineIdRegistryRegistry>::TearDown();
@@ -34,10 +44,52 @@ void FOnlineIdRegistryRegistry::UnregisterAccountIdRegistry(EOnlineServices Onli
 	}
 }
 
-IOnlineAccountIdRegistry* FOnlineIdRegistryRegistry::GetAccountIdRegistry(EOnlineServices OnlineServices)
+FString FOnlineIdRegistryRegistry::ToLogString(const FOnlineAccountIdHandle& Handle) const
+{
+	FString Result;
+	if (IOnlineAccountIdRegistry* Registry = GetAccountIdRegistry(Handle.GetOnlineServicesType()))
+	{
+		Result = Registry->ToLogString(Handle);
+	}
+	else
+	{
+		Result = ForeignAccountIdRegistry->ToLogString(Handle);
+	}
+	return Result;
+}
+
+TArray<uint8> FOnlineIdRegistryRegistry::ToReplicationData(const FOnlineAccountIdHandle& Handle) const
+{
+	TArray<uint8> RepData;
+	if (IOnlineAccountIdRegistry* Registry = GetAccountIdRegistry(Handle.GetOnlineServicesType()))
+	{
+		RepData = Registry->ToReplicationData(Handle);
+	}
+	else
+	{
+		RepData = ForeignAccountIdRegistry->ToReplicationData(Handle);
+	}
+	return RepData;
+}
+
+FOnlineAccountIdHandle FOnlineIdRegistryRegistry::ToAccountId(EOnlineServices Services, const TArray<uint8>& RepData) const
+{
+	FOnlineAccountIdHandle Handle;
+	if (IOnlineAccountIdRegistry* Registry = GetAccountIdRegistry(Handle.GetOnlineServicesType()))
+	{
+		Handle = Registry->FromReplicationData(RepData);
+	}
+	else
+	{
+		Handle = ForeignAccountIdRegistry->FromReplicationData(Services, RepData);
+	}
+	return Handle;
+}
+
+IOnlineAccountIdRegistry* FOnlineIdRegistryRegistry::GetAccountIdRegistry(EOnlineServices OnlineServices) const
 {
 	IOnlineAccountIdRegistry* Registry = nullptr;
-	if (FAccountIdRegistryAndPriority* Found = AccountIdRegistries.Find(OnlineServices))
+	if (const FAccountIdRegistryAndPriority* Found = AccountIdRegistries.Find(OnlineServices))
 	{
 		Registry = Found->Registry;
 	}
@@ -46,12 +98,7 @@ IOnlineAccountIdRegistry* FOnlineIdRegistryRegistry::GetAccountIdRegistry(EOnlin
 
 FString ToLogString(const FOnlineAccountIdHandle& Id)
 {
-	FString Result;
-	if (IOnlineAccountIdRegistry* Registry = FOnlineIdRegistryRegistry::Get().GetAccountIdRegistry(Id.GetOnlineServicesType()))
-	{
-		Result = Registry->ToLogString(Id);
-	}
-	return Result;
+	return FOnlineIdRegistryRegistry::Get().ToLogString(Id);
 }
 
 }	/* UE::Online */
@@ -61,19 +108,14 @@ FString FUniqueNetIdWrapper::ToDebugString() const
 	FString Result;
 	if (IsValid())
 	{
-		if (Variant.IsType<FUniqueNetIdPtr>())
+		if (IsV1())
 		{
-			const FUniqueNetIdPtr& Ptr = Variant.Get<FUniqueNetIdPtr>();
+			const FUniqueNetIdPtr& Ptr = GetV1();
 			Result = FString::Printf(TEXT("%s:%s"), *Ptr->GetType().ToString(), *Ptr->ToDebugString());
 		}
-		else if (Variant.IsType<UE::Online::FOnlineAccountIdHandle>())
-		{
-			const UE::Online::FOnlineAccountIdHandle& Handle = Variant.Get<UE::Online::FOnlineAccountIdHandle>();
-			Result = ToLogString(Handle);
-		}
 		else
 		{
-			checkNoEntry();
+			Result = ToLogString(GetV2());
 		}
 	}
 	else
diff --git a/Engine/Source/Runtime/CoreOnline/Private/Online/CoreOnlinePrivate.cpp b/Engine/Source/Runtime/CoreOnline/Private/Online/CoreOnlinePrivate.cpp
new file mode 100644
--- /dev/null
+++ b/Engine/Source/Runtime/CoreOnline/Private/Online/CoreOnlinePrivate.cpp
@@ -0,0 +1,69 @@
+// Copyright Epic Games, Inc. All Rights Reserved.
+
+#include "Online/CoreOnlinePrivate.h"
+
+namespace UE::Online {
+
+FString FOnlineForeignAccountIdRegistry::ToLogString(const FOnlineAccountIdHandle& Handle) const
+{
+	FString Result;
+	if (Handle.IsValid())
+	{
+		const EOnlineServices HandleType = Handle.GetOnlineServicesType();
+		const uint32 HandleIndex = Handle.GetHandle() - 1;
+		if (const FRepData* RepDataForServices = OnlineServicesToRepData.Find(HandleType))
+		{
+			const TArray<TArray<uint8>>& RepDataArray = RepDataForServices->RepDataArray;
+			if (ensure(RepDataArray.IsValidIndex(HandleIndex)))
+			{
+				const TArray<uint8>& RepData = RepDataArray[HandleIndex];
+				Result = FString::Printf(TEXT("ForeignId=[Type=%d Handle=%d RepData=[%s]"),
+					HandleType,
+					HandleIndex,
+					*FString::FromHexBlob(RepData.GetData(), RepData.Num()));
+			}
+		}
+	}
+	return Result;
+}
+
+TArray<uint8> FOnlineForeignAccountIdRegistry::ToReplicationData(const FOnlineAccountIdHandle& Handle) const
+{
+	TArray<uint8> RepData;
+	if (Handle.IsValid())
+	{
+		const EOnlineServices HandleType = Handle.GetOnlineServicesType();
+		const uint32 HandleIndex = Handle.GetHandle() - 1;
+		if (const FRepData* RepDataForServices = OnlineServicesToRepData.Find(Handle.GetOnlineServicesType()))
+		{
+			const TArray<TArray<uint8>>& RepDataArray = RepDataForServices->RepDataArray;
+			if (ensure(RepDataArray.IsValidIndex(HandleIndex)))
+			{
+				RepData = RepDataArray[HandleIndex];
+			}
+		}
+	}
+	return RepData;
+}
+
+FOnlineAccountIdHandle FOnlineForeignAccountIdRegistry::FromReplicationData(EOnlineServices Services, const TArray<uint8>& RepData)
+{
+	FOnlineAccountIdHandle Handle;
+	if (RepData.Num())
+	{
+		FRepData& RepDataForServices = OnlineServicesToRepData.FindOrAdd(Services);
+		if (FOnlineAccountIdHandle* Found = RepDataForServices.RepDataToHandle.Find(RepData))
+		{
+			Handle = *Found;
+		}
+		else
+		{
+			RepDataForServices.RepDataArray.Add(RepData);
+			Handle = FOnlineAccountIdHandle(Services, RepDataForServices.RepDataArray.Num());
+			RepDataForServices.RepDataToHandle.Emplace(RepData, Handle);
+		}
+	}
+	return Handle;
+}
+
+}
\ No newline at end of file
diff --git a/Engine/Source/Runtime/CoreOnline/Private/Online/CoreOnlinePrivate.h b/Engine/Source/Runtime/CoreOnline/Private/Online/CoreOnlinePrivate.h
new file mode 100644
--- /dev/null
+++ b/Engine/Source/Runtime/CoreOnline/Private/Online/CoreOnlinePrivate.h
@@ -0,0 +1,44 @@
+// Copyright Epic Games, Inc. All Rights Reserved.
+
+#pragma once
+
+#include "Online/CoreOnline.h"
+#include "UObject/ObjectMacros.h"
+
+// This exists purely to guarantee a package is created, or the engine will not boot.
+UENUM(BlueprintType)
+enum class ECoreOnlineDummy : uint8
+{
+	Dummy
+};
+
+static uint32 GetTypeHash(const TArray<uint8>& Array)
+{
+	uint32 Seed = 0;
+	for (const uint8& Elem : Array)
+		Seed ^= GetTypeHash(Elem) + 0x9e3779b9 + (Seed << 6) + (Seed >> 2);
+	return Seed;
+}
+
+namespace UE::Online {
+
+/**
+ * Registry for storing account id's for unregistered services implementations.
+ */
+class FOnlineForeignAccountIdRegistry
+{
+public:
+	FString ToLogString(const FOnlineAccountIdHandle& Handle) const;
+	TArray<uint8> ToReplicationData(const FOnlineAccountIdHandle& Handle) const;
+	FOnlineAccountIdHandle FromReplicationData(EOnlineServices Services, const TArray<uint8>& RepData);
+
+private:
+	struct FRepData
+	{
+		TMap<TArray<uint8>, FOnlineAccountIdHandle> RepDataToHandle;
+		TArray<TArray<uint8>> RepDataArray;
+	};
+	TMap<EOnlineServices, FRepData> OnlineServicesToRepData;
+};
+
+} // namespace UE::Online
\ No newline at end of file
diff --git a/Engine/Source/Runtime/CoreOnline/Private/Tests/CoreOnlineTests.cpp b/Engine/Source/Runtime/CoreOnline/Private/Tests/CoreOnlineTests.cpp
new file mode 100644
--- /dev/null
+++ b/Engine/Source/Runtime/CoreOnline/Private/Tests/CoreOnlineTests.cpp
@@ -0,0 +1,22 @@
+// Copyright Epic Games, Inc. All Rights Reserved.
+
+#include "Misc/AutomationTest.h"
+#include "Online/CoreOnline.h"
+#include "Online/CoreOnlinePrivate.h"
+
+IMPLEMENT_SIMPLE_AUTOMATION_TEST(FForeignIdRepTest,
+	"System.Engine.Online.ForeignIdRepTest",
+	EAutomationTestFlags::ApplicationContextMask | EAutomationTestFlags::SmokeFilter)
+	bool FForeignIdRepTest::RunTest(const FString& Parameters)
+{
+	using namespace UE::Online;
+	const EOnlineServices Services = EOnlineServices::Apple;
+	const TArray<uint8> SourceRepData { 0, 1, 2, 3, 4, 5, 6, 7 };
+
+	FOnlineIdRegistryRegistry Registry;
+	const FOnlineAccountIdHandle Handle = Registry.ToAccountId(Services, SourceRepData);
+	const TArray<uint8> RepData = Registry.ToReplicationData(Handle);
+	UTEST_EQUAL(TEXT(""), RepData, SourceRepData);
+
+	return true;
+}
\ No newline at end of file
diff --git a/Engine/Source/Runtime/CoreOnline/Public/Online/CoreOnline.h b/Engine/Source/Runtime/CoreOnline/Public/Online/CoreOnline.h
--- a/Engine/Source/Runtime/CoreOnline/Public/Online/CoreOnline.h
+++ b/Engine/Source/Runtime/CoreOnline/Public/Online/CoreOnline.h
@@ -4,10 +4,12 @@
 
 #include "CoreMinimal.h"
 #include "Hash/CityHash.h"
-#include "Online/CoreOnlineFwd.h"
 #include "Misc/TVariant.h"
+#include "Online/CoreOnlineFwd.h"
 #include "UObject/ObjectMacros.h"
 
+#include "Online/CoreOnlinePackage.h"
+
 class FLazySingleton;
 
 UE_DEPRECATED(5.0, "Use NAME_GameSession.")
@@ -177,6 +179,8 @@ public:
 
 namespace UE::Online {
 
+class FOnlineForeignAccountIdRegistry;
+
 /** Tags used as template argument to TOnlineIdHandle to make it a compile error to assign between id's of different types */
 namespace OnlineIdHandleTags
 {
@@ -286,7 +290,7 @@ class IOnlineIdRegistry
 public:
 	virtual FString ToLogString(const TOnlineIdHandle<IdType>& Handle) const = 0;
 	virtual TArray<uint8> ToReplicationData(const TOnlineIdHandle<IdType>& Handle) const = 0;
-	virtual TOnlineIdHandle<IdType> FromReplicationData(const TArrayView<uint8> Handle) = 0;
+	virtual TOnlineIdHandle<IdType> FromReplicationData(const TArray<uint8> Handle) = 0;
 };
 
 using IOnlineAccountIdRegistry = IOnlineIdRegistry<OnlineIdHandleTags::FAccount>;
@@ -305,6 +309,7 @@ public:
 	 * Tear down the singleton instance
 	 */
 	COREONLINE_API static void TearDown();
+
 	/**
 	 * Register a registry for a given OnlineServices implementation and TOnlineIdHandle type
 	 *
@@ -322,16 +327,13 @@ public:
 	 */
 	COREONLINE_API void UnregisterAccountIdRegistry(EOnlineServices OnlineServices, int32 Priority = 0);
 
-	/**
-	 * Get the account id registry
-	 *
-	 * @param OnlineServices Type of online services for the IOnlineServices instance
-	 *
-	 * @return The account id registry, or an invalid pointer if it is not registered
-	 */
-	COREONLINE_API IOnlineAccountIdRegistry* GetAccountIdRegistry(EOnlineServices OnlineServices);
+	COREONLINE_API FString ToLogString(const FOnlineAccountIdHandle& Handle) const;
+	COREONLINE_API TArray<uint8> ToReplicationData(const FOnlineAccountIdHandle& Handle) const;
+	COREONLINE_API FOnlineAccountIdHandle ToAccountId(EOnlineServices Services, const TArray<uint8>& RepData) const;
 
 private:
+	IOnlineAccountIdRegistry* GetAccountIdRegistry(EOnlineServices OnlineServices) const;
+
 	struct FAccountIdRegistryAndPriority
 	{
 		FAccountIdRegistryAndPriority(IOnlineAccountIdRegistry* InRegistry, int32 InPriority)
@@ -342,10 +344,13 @@ private:
 	};
 
 	TMap<EOnlineServices, FAccountIdRegistryAndPriority> AccountIdRegistries;
+	TUniquePtr<FOnlineForeignAccountIdRegistry> ForeignAccountIdRegistry;
 
-	FOnlineIdRegistryRegistry() = default;
-	~FOnlineIdRegistryRegistry() = default;
 	friend FLazySingleton;
+
+PACKAGE_SCOPE:
+	FOnlineIdRegistryRegistry();
+	~FOnlineIdRegistryRegistry();
 };
 
 }	/* UE::Online */
diff --git a/Engine/Source/Runtime/CoreOnline/Public/Online/CoreOnlinePackage.h b/Engine/Source/Runtime/CoreOnline/Public/Online/CoreOnlinePackage.h
new file mode 100644
--- /dev/null
+++ b/Engine/Source/Runtime/CoreOnline/Public/Online/CoreOnlinePackage.h
@@ -0,0 +1,10 @@
+// Copyright Epic Games, Inc. All Rights Reserved.
+
+// [[ IncludeTool: Inline ]] // Markup to tell IncludeTool that this file is state changing and cannot be optimized out.
+
+#undef PACKAGE_SCOPE
+#ifdef COREONLINE_PACKAGE
+#define PACKAGE_SCOPE public
+#else
+#define PACKAGE_SCOPE protected
+#endif
diff --git a/Engine/Source/Runtime/Engine/Private/OnlineReplStructs.cpp b/Engine/Source/Runtime/Engine/Private/OnlineReplStructs.cpp
--- a/Engine/Source/Runtime/Engine/Private/OnlineReplStructs.cpp
+++ b/Engine/Source/Runtime/Engine/Private/OnlineReplStructs.cpp
@@ -229,14 +229,8 @@ void FUniqueNetIdRepl::MakeReplicationDataV2()
 
 	UE::Online::FOnlineAccountIdHandle Handle = GetV2();
 	UE::Online::EOnlineServices OnlineServicesType = Handle.GetOnlineServicesType();
-
-	TArray<uint8> ReplicationData;
-	const UE::Online::IOnlineAccountIdRegistry* Registry = UE::Online::FOnlineIdRegistryRegistry::Get().GetAccountIdRegistry(OnlineServicesType);
-	if (ensure(Registry))
-	{
-		ReplicationData = Registry->ToReplicationData(Handle);
-		check(!ReplicationData.IsEmpty());
-	}
+	TArray<uint8> ReplicationData = UE::Online::FOnlineIdRegistryRegistry::Get().ToReplicationData(Handle);
+	check(!ReplicationData.IsEmpty());
 
 	ReplicationBytes.Empty();
 	FMemoryWriter Writer(ReplicationBytes);
@@ -451,13 +445,9 @@ void FUniqueNetIdRepl::NetSerializeLoadV2(FArchive& Ar, const EUniqueIdEncodingF
 	TArray<uint8> ReplicationData;
 	Ar << ReplicationData;
 
-	UE::Online::IOnlineAccountIdRegistry* Registry = UE::Online::FOnlineIdRegistryRegistry::Get().GetAccountIdRegistry(OnlineServicesType);
-	if (ensure(Registry))
-	{
-		const UE::Online::FOnlineAccountIdHandle Handle = Registry->FromReplicationData(ReplicationData);
-		check(Handle.IsValid());
-		SetAccountId(Handle);
-	}
+	const UE::Online::FOnlineAccountIdHandle Handle = UE::Online::FOnlineIdRegistryRegistry::Get().ToAccountId(OnlineServicesType, ReplicationData);
+	check(Handle.IsValid());
+	SetAccountId(Handle);
 }
 
 bool FUniqueNetIdRepl::Serialize(FArchive& Ar)
