OSSv2 Lobby dependencies.

diff --git a/Engine/Plugins/Online/Experimental/OnlineServices/Source/OnlineServicesCommon/Public/Online/OnlineAsyncOp.h b/Engine/Plugins/Online/Experimental/OnlineServices/Source/OnlineServicesCommon/Public/Online/OnlineAsyncOp.h
--- a/Engine/Plugins/Online/Experimental/OnlineServices/Source/OnlineServicesCommon/Public/Online/OnlineAsyncOp.h
+++ b/Engine/Plugins/Online/Experimental/OnlineServices/Source/OnlineServicesCommon/Public/Online/OnlineAsyncOp.h
@@ -47,6 +47,43 @@ private:
 	EOnlineAsyncExecutionPolicy ExecutionPolicy;
 };
 
+/* For use when we need to put a TOnlineResult in an object that needs to be default constructed such as a TPromise */
+template <typename Result>
+class TDefaultErrorResultInternal : public TResult<Result, FOnlineError>
+{
+public:
+	using TResult<Result, FOnlineError>::TResult;
+
+	TDefaultErrorResultInternal()
+		: TResult<Result, FOnlineError>(Errors::Unknown())
+	{
+	}
+};
+
+template <typename OpType>
+class TDefaultErrorResult : public TOnlineResult<OpType>
+{
+public:
+	using TOnlineResult<OpType>::TOnlineResult;
+
+	TDefaultErrorResult()
+		: TOnlineResult<OpType>(Errors::Unknown())
+	{
+	}
+};
+
+template <typename T>
+FString ToLogString(const TDefaultErrorResultInternal<T>& Result)
+{
+	if (Result.IsOk())
+	{
+		return ToLogString(Result.GetOkValue());
+	}
+	else
+	{
+		return ToLogString(Result.GetErrorValue());
+	}
+}
 
 namespace Private
 {
diff --git a/Engine/Plugins/Online/Experimental/OnlineServices/Source/OnlineServicesCommon/Public/Online/OnlineConfig.h b/Engine/Plugins/Online/Experimental/OnlineServices/Source/OnlineServicesCommon/Public/Online/OnlineConfig.h
--- a/Engine/Plugins/Online/Experimental/OnlineServices/Source/OnlineServicesCommon/Public/Online/OnlineConfig.h
+++ b/Engine/Plugins/Online/Experimental/OnlineServices/Source/OnlineServicesCommon/Public/Online/OnlineConfig.h
@@ -193,6 +193,32 @@ public:
 		return false;
 	}
 
+	/**
+	 * Get a value consisting of a TArray of FText
+	 *
+	 * @param Section Section to read the value from
+	 * @param Key Key in the section for the value
+	 * @param Value Array of FText read from the config. Will be empty if not present
+	 *
+	 * @return number of values in the array
+	 */
+	virtual int32 GetValue(const TCHAR* Section, const TCHAR* Key, TArray<FText>& Value)
+	{
+		Value.Empty();
+		TArray<FString> StringArray;
+		if (GetValue(Section, Key, StringArray) > 0)
+		{
+			Value.Reserve(StringArray.Num());
+			for (const FString& StringValue : StringArray)
+			{
+				FText TextValue;
+				FTextStringHelper::ReadFromBuffer(*StringValue, TextValue, Section);
+				Value.Emplace(TextValue);
+			}
+		}
+		return Value.Num();
+	}
+
 	/**
 	 * Get a value consisting of a TArray of bool
 	 *
diff --git a/Engine/Plugins/Online/Experimental/OnlineServices/Source/OnlineServicesCommon/Public/Online/OnlineExecHandler.h b/Engine/Plugins/Online/Experimental/OnlineServices/Source/OnlineServicesCommon/Public/Online/OnlineExecHandler.h
--- a/Engine/Plugins/Online/Experimental/OnlineServices/Source/OnlineServicesCommon/Public/Online/OnlineExecHandler.h
+++ b/Engine/Plugins/Online/Experimental/OnlineServices/Source/OnlineServicesCommon/Public/Online/OnlineExecHandler.h
@@ -2,6 +2,7 @@
 
 #pragma once
 
+#include "Online/CoreOnline.h"
 #include "Online/OnlineUtils.h"
 #include "Online/OnlineServicesLog.h"
 
@@ -57,6 +58,21 @@ struct TOnlineInterfaceOperationMemberFunctionPtrTraits<TOnlineAsyncOpHandle<TOp
 	static constexpr bool bAsync = true;
 };
 
+// Forward declarations for dependent type resolution
+inline bool ParseOnlineExecParams(const TCHAR*& Cmd, FString& Value);
+inline bool ParseOnlineExecParams(const TCHAR*& Cmd, FName& Value);
+inline bool ParseOnlineExecParams(const TCHAR*& Cmd, uint8& Value);
+inline bool ParseOnlineExecParams(const TCHAR*& Cmd, int32& Value);
+inline bool ParseOnlineExecParams(const TCHAR*& Cmd, uint32& Value);
+inline bool ParseOnlineExecParams(const TCHAR*& Cmd, int64& Value);
+inline bool ParseOnlineExecParams(const TCHAR*& Cmd, uint64& Value);
+template <typename T> inline bool ParseOnlineExecParams(const TCHAR*& Cmd, TArray<T>& Array);
+template <typename T, typename U> inline bool ParseOnlineExecParams(const TCHAR*& Cmd, TMap<T, U>& Map);
+template <typename... Ts> inline bool ParseOnlineExecParams(const TCHAR*& Cmd, TVariant<Ts...>& Variant);
+template <typename IdType> inline bool ParseOnlineExecParams(const TCHAR*& Cmd, TOnlineIdHandle<IdType>& Value);
+template <typename T> std::enable_if_t<!TModels<Meta::COnlineMetadataAvailable, T>::Value, bool> ParseOnlineExecParams(const TCHAR*& Cmd, T& Value);
+template <typename T> std::enable_if_t<TModels<Meta::COnlineMetadataAvailable, T>::Value, bool> ParseOnlineExecParams(const TCHAR*& Cmd, T& Value);
+
 inline bool ParseOnlineExecParams(const TCHAR*& Cmd, FString& Value)
 {
 	if (!FParse::Token(Cmd, Value, true))
@@ -67,12 +83,24 @@ inline bool ParseOnlineExecParams(const TCHAR*& Cmd, FString& Value)
 	return true;
 }
 
+inline bool ParseOnlineExecParams(const TCHAR*& Cmd, FName& Value)
+{
+	FString StringValue;
+	if (!ParseOnlineExecParams(Cmd, StringValue))
+	{
+		return false;
+	}
+
+	Value = FName(StringValue);
+	return true;
+}
+
 inline bool ParseOnlineExecParams(const TCHAR*& Cmd, uint8& Value)
 {
 	FString Token;
 	if (FParse::Token(Cmd, Token, true))
 	{
-		Value = static_cast<uint8>(FCString::Strtoui64(Cmd, nullptr, 10));
+		Value = static_cast<uint8>(FCString::Strtoui64(*Token, nullptr, 10));
 	}
 	else
 	{
@@ -87,7 +115,7 @@ inline bool ParseOnlineExecParams(const TCHAR*& Cmd, int32& Value)
 	FString Token;
 	if (FParse::Token(Cmd, Token, true))
 	{
-		Value = FCString::Strtoi(Cmd, nullptr, 10);
+		Value = FCString::Strtoi(*Token, nullptr, 10);
 	}
 	else
 	{
@@ -102,7 +130,7 @@ inline bool ParseOnlineExecParams(const TCHAR*& Cmd, uint32& Value)
 	FString Token;
 	if (FParse::Token(Cmd, Token, true))
 	{
-		Value = static_cast<uint32>(FCString::Strtoui64(Cmd, nullptr, 10));
+		Value = static_cast<uint32>(FCString::Strtoui64(*Token, nullptr, 10));
 	}
 	else
 	{
@@ -117,7 +145,7 @@ inline bool ParseOnlineExecParams(const TCHAR*& Cmd, int64& Value)
 	FString Token;
 	if (FParse::Token(Cmd, Token, true))
 	{
-		Value = FCString::Strtoi64(Cmd, nullptr, 10);
+		Value = FCString::Strtoi64(*Token, nullptr, 10);
 	}
 	else
 	{
@@ -132,7 +160,7 @@ inline bool ParseOnlineExecParams(const TCHAR*& Cmd, uint64& Value)
 	FString Token;
 	if (FParse::Token(Cmd, Token, true))
 	{
-		Value = FCString::Strtoui64(Cmd, nullptr, 10);
+		Value = FCString::Strtoui64(*Token, nullptr, 10);
 	}
 	else
 	{
@@ -183,6 +211,46 @@ inline bool ParseOnlineExecParams(const TCHAR*& Cmd, TArray<T>& Array)
 	return true;
 }
 
+template <typename T, typename U>
+inline bool ParseOnlineExecParams(const TCHAR*& Cmd, TMap<T, U>& Map)
+{
+	FString Token;
+	if (FParse::Token(Cmd, Token, true))
+	{
+		TArray<FString> TokenArray;
+		Token.ParseIntoArray(TokenArray, TEXT(","));
+		Map.Reserve(TokenArray.Num());
+		for (const FString& ArrayToken : TokenArray)
+		{
+			TArray<FString> TokenValuePair;
+			Token.ParseIntoArray(TokenValuePair, TEXT("="));
+
+			if (TokenValuePair.Num() > 1)
+			{
+				T Key;
+				U Value;
+
+				const TCHAR* ArrayTokenKeyTCHAR = *TokenValuePair[0];
+				if (ParseOnlineExecParams(ArrayTokenKeyTCHAR, Key))
+				{
+					if (TokenValuePair.Num() > 2)
+					{
+						const TCHAR* ArrayTokenValueTCHAR = *TokenValuePair[1];
+						ParseOnlineExecParams(ArrayTokenValueTCHAR, Value);
+
+						Map.Emplace(MoveTempIfPossible(Key), MoveTempIfPossible(Value));
+					}
+				}
+			}
+		}
+	}
+	else
+	{
+		return false;
+	}
+	return true;
+}
+
 template <typename... Ts>
 inline bool ParseOnlineExecParams(const TCHAR*& Cmd, TVariant<Ts...>& Variant)
 {
@@ -191,10 +259,17 @@ inline bool ParseOnlineExecParams(const TCHAR*& Cmd, TVariant<Ts...>& Variant)
 }
 
 template<typename IdType>
-inline bool ParseOnlineExecParams(const TCHAR*& Cmd, TOnlineIdHandle<IdType>& Id)
+inline bool ParseOnlineExecParams(const TCHAR*& Cmd, TOnlineIdHandle<IdType>& Value)
 {
-	// TODO
-	return true;
+	EOnlineServices Type;
+	uint32 Handle;
+	if (ParseOnlineExecParams(Cmd, Type) && ParseOnlineExecParams(Cmd, Handle))
+	{
+		Value = TOnlineIdHandle<IdType>(Type, Handle);
+		return true;
+	}
+
+	return false;
 }
 
 template <typename T>
diff --git a/Engine/Plugins/Online/Experimental/OnlineServices/Source/OnlineServicesInterface/Public/Online/OnlineResult.h b/Engine/Plugins/Online/Experimental/OnlineServices/Source/OnlineServicesInterface/Public/Online/OnlineResult.h
--- a/Engine/Plugins/Online/Experimental/OnlineServices/Source/OnlineServicesInterface/Public/Online/OnlineResult.h
+++ b/Engine/Plugins/Online/Experimental/OnlineServices/Source/OnlineServicesInterface/Public/Online/OnlineResult.h
@@ -189,9 +189,10 @@ public:
 	}
 
 
-private:
-	TResult() = delete;
+protected:
+	TResult() = default;
 
+private:
 	/** Location that the result's value is stored */
 	TVariant<SuccessType, ErrorType> Storage;
 };
diff --git a/Engine/Plugins/Online/Experimental/OnlineServices/Source/OnlineServicesInterface/Public/Online/OnlineUtils.h b/Engine/Plugins/Online/Experimental/OnlineServices/Source/OnlineServicesInterface/Public/Online/OnlineUtils.h
--- a/Engine/Plugins/Online/Experimental/OnlineServices/Source/OnlineServicesInterface/Public/Online/OnlineUtils.h
+++ b/Engine/Plugins/Online/Experimental/OnlineServices/Source/OnlineServicesInterface/Public/Online/OnlineUtils.h
@@ -8,6 +8,19 @@ namespace UE::Online {
 
 template <typename T> FString ToLogString(const TArray<T>& Array);
 template <typename K, typename V> FString ToLogString(const TMap<K, V>& Map);
+template <typename T, ESPMode Mode> FString ToLogString(const TSharedPtr<T, Mode>& Ptr);
+template <typename T> FString ToLogString(const TOptional<T> Optional);
+template <typename... Ts> FString ToLogString(const TVariant<Ts...>& Variant);
+template <typename T, ESPMode Mode> FString ToLogString(const TSharedRef<T, Mode>& Ref);
+inline FString ToLogString(const FString& String);
+inline FString ToLogString(uint8 Value);
+inline FString ToLogString(int8 Value);
+inline FString ToLogString(uint16 Value);
+inline FString ToLogString(int16 Value);
+inline FString ToLogString(uint32 Value);
+inline FString ToLogString(int32 Value);
+inline FString ToLogString(uint64 Value);
+inline FString ToLogString(int64 Value);
 template <typename T> std::enable_if_t<!TModels<Meta::COnlineMetadataAvailable, T>::Value, FString> ToLogString(const T& Value);
 template <typename T> std::enable_if_t<TModels<Meta::COnlineMetadataAvailable, T>::Value, FString> ToLogString(const T& Value);
 
diff --git a/Engine/Plugins/Online/Experimental/OnlineServicesEOS/Source/Private/OnlineServicesEOS.cpp b/Engine/Plugins/Online/Experimental/OnlineServicesEOS/Source/Private/OnlineServicesEOS.cpp
--- a/Engine/Plugins/Online/Experimental/OnlineServicesEOS/Source/Private/OnlineServicesEOS.cpp
+++ b/Engine/Plugins/Online/Experimental/OnlineServicesEOS/Source/Private/OnlineServicesEOS.cpp
@@ -8,6 +8,7 @@
 
 #include "AuthEOS.h"
 #include "FriendsEOS.h"
+#include "LobbiesEOS.h"
 #include "OnlineIdEOS.h"
 #include "PresenceEOS.h"
 #include "ExternalUIEOS.h"
@@ -44,6 +45,7 @@ void FOnlineServicesEOS::RegisterComponents()
 {
 	Components.Register<FAuthEOS>(*this);
 	Components.Register<FFriendsEOS>(*this);
+	Components.Register<FLobbiesEOS>(*this);
 	Components.Register<FPresenceEOS>(*this);
 	Components.Register<FExternalUIEOS>(*this);
 	FOnlineServicesCommon::RegisterComponents();
@@ -61,7 +63,14 @@ void FOnlineServicesEOS::Initialize()
 	FTCHARToUTF8 ClientSecret(*EOSPlatformConfig.ClientSecret);
 
 	// Init EOS SDK
+	static struct FReservedOptions
+	{
+		int32_t ApiVersion;
+		const char* BackendEnvironment;
+	} ReservedOptions = { 1, "GameDev" };
+
 	EOS_Platform_Options PlatformOptions = {};
+	PlatformOptions.Reserved = &ReservedOptions;
 	PlatformOptions.ApiVersion = EOS_PLATFORM_OPTIONS_API_LATEST;
 	PlatformOptions.bIsServer = EOS_FALSE;
 	PlatformOptions.OverrideCountryCode = nullptr;
@@ -80,11 +89,19 @@ void FOnlineServicesEOS::Initialize()
 	PlatformOptions.ClientCredentials.ClientSecret = ClientSecret.Get();
 
 	IEOSSDKManager* SDKManager = IEOSSDKManager::Get();
-	if (ensure(SDKManager))
+	if (!SDKManager)
+	{
+		return;
+	}
+
+	EOS_EResult InitResult = SDKManager->Initialize();
+	if (InitResult != EOS_EResult::EOS_Success)
 	{
-		EOSPlatformHandle = SDKManager->CreatePlatform(PlatformOptions);
+		return;
 	}
 
+	EOSPlatformHandle = SDKManager->CreatePlatform(PlatformOptions);
+
 	FOnlineServicesCommon::Initialize();
 }
 
diff --git a/Engine/Plugins/Online/Experimental/OnlineServicesEOS/Source/Public/OnlineServicesEOSTypes.h b/Engine/Plugins/Online/Experimental/OnlineServicesEOS/Source/Public/OnlineServicesEOSTypes.h
--- a/Engine/Plugins/Online/Experimental/OnlineServicesEOS/Source/Public/OnlineServicesEOSTypes.h
+++ b/Engine/Plugins/Online/Experimental/OnlineServicesEOS/Source/Public/OnlineServicesEOSTypes.h
@@ -3,12 +3,19 @@
 #pragma once
 
 #include "Online/OnlineAsyncOp.h"
+#include "Online/OnlineError.h"
 #include "Async/Future.h"
 
 #include "EOSShared.h"
 
 namespace UE::Online {
 
+inline FOnlineError FromEOSError(EOS_EResult ResultCode)
+{
+	// Todo: make this for real.
+	return FOnlineError(Errors::Unknown());
+}
+
 /** Class to handle all callbacks generically using a future to forward callback results */
 template<typename CallbackType>
 class TEOSCallback
@@ -62,4 +69,151 @@ void EOS_Async(TEOSFn EOSFn, TEOSHandle EOSHandle, TEOSParameters Parameters, TP
 	EOSFn(EOSHandle, &Parameters, Callback, *Callback);
 }
 
+class EOSEventRegistration
+{
+public:
+	virtual ~EOSEventRegistration() = default;
+};
+typedef TUniquePtr<EOSEventRegistration> EOSEventRegistrationPtr;
+
+/** 
+* EOS event registration utility for binding an EOS notifier registration to a RAII object which handles
+* unregistering when it exits scope. Intended to be used from a TOnlineComponent class.
+* 
+* Example:
+*	EOSEventRegistrationPtr OnLobbyUpdatedEOSEventRegistration = EOS_RegisterComponentEventHandler(
+*		this,
+*		LobbyHandle,
+*		EOS_LOBBY_ADDNOTIFYLOBBYUPDATERECEIVED_API_LATEST,
+*		&EOS_Lobby_AddNotifyLobbyUpdateReceived,
+*		&EOS_Lobby_RemoveNotifyLobbyUpdateReceived,
+*		&FLobbiesEOS::HandleLobbyUpdated);
+**/
+template <
+	typename ComponentHandlerClass,
+	typename EOSHandle,
+	typename EOSNotfyRegisterFunction,
+	typename EOSNotfyUnregisterFunction,
+	typename ComponentHandlerFunction>
+EOSEventRegistrationPtr EOS_RegisterComponentEventHandler(
+	ComponentHandlerClass* HandlerClass,
+	EOSHandle ClientHandle,
+	int32_t ApiVersion,
+	EOSNotfyRegisterFunction NotfyRegisterFunction,
+	EOSNotfyUnregisterFunction NotfyUnregisterFunction,
+	ComponentHandlerFunction HandlerFunction);
+
+namespace detail {
+
+template<typename Function> struct TEOSCallbackTraitsBase;
+
+template<typename Component, typename EventData>
+struct TEOSCallbackTraitsBase<void (Component::*)(const EventData*)>
+{
+	typedef Component ComponentType;
+	typedef EventData EventDataType;
+};
+
+template<typename Function>
+struct TEOSCallbackTraits : public TEOSCallbackTraitsBase<typename std::remove_reference<Function>::type>
+{
+};
+
+template<typename Function> struct TEOSNotifyRegisterTraits;
+
+template<typename EOSHandle, typename Options, typename ClientData, typename NotificationFn, typename NotificationId>
+struct TEOSNotifyRegisterTraits<NotificationId (*)(EOSHandle, const Options*, ClientData*, NotificationFn)>
+{
+	typedef Options OptionsType;
+	typedef NotificationId NotificationIdType;
+};
+
+template <
+	typename ComponentHandlerClass,
+	typename EOSHandle,
+	typename EOSNotfyRegisterFunction,
+	typename EOSNotfyUnregisterFunction,
+	typename ComponentHandlerFunction>
+class EOSEventRegistrationImpl : public EOSEventRegistration
+{
+public:
+	typedef EOSEventRegistrationImpl<
+		ComponentHandlerClass,
+		EOSHandle,
+		EOSNotfyRegisterFunction,
+		EOSNotfyUnregisterFunction,
+		ComponentHandlerFunction> ThisClass;
+
+	EOSEventRegistrationImpl(
+		ComponentHandlerClass* HandlerClass,
+		EOSHandle ClientHandle,
+		int32_t ApiVersion,
+		EOSNotfyRegisterFunction NotfyRegisterFunction,
+		EOSNotfyUnregisterFunction NotfyUnregisterFunction,
+		ComponentHandlerFunction HandlerFunction)
+		: HandlerClass(HandlerClass)
+		, ClientHandle(ClientHandle)
+		, NotfyUnregisterFunction(NotfyUnregisterFunction)
+		, HandlerFunction(HandlerFunction)
+	{
+		typename TEOSNotifyRegisterTraits<EOSNotfyRegisterFunction>::OptionsType Options = { };
+		Options.ApiVersion = ApiVersion;
+		NotificationId = NotfyRegisterFunction(ClientHandle, &Options, this,
+		[](const typename TEOSCallbackTraits<ComponentHandlerFunction>::EventDataType* Data)
+		{
+			ThisClass* This = reinterpret_cast<ThisClass*>(Data->ClientData);
+			(This->HandlerClass->*This->HandlerFunction)(Data);
+		});
+	}
+
+	virtual ~EOSEventRegistrationImpl()
+	{
+		NotfyUnregisterFunction(ClientHandle, NotificationId);
+	}
+
+private:
+	EOSEventRegistrationImpl() = delete;
+	EOSEventRegistrationImpl(const EOSEventRegistrationImpl&) = delete;
+	EOSEventRegistrationImpl& operator=(const EOSEventRegistrationImpl&) = delete;
+
+	typedef typename TEOSNotifyRegisterTraits<EOSNotfyRegisterFunction>::NotificationIdType NotificationIdType;
+
+	NotificationIdType NotificationId;
+	ComponentHandlerClass* HandlerClass;
+	EOSHandle ClientHandle;
+	EOSNotfyUnregisterFunction NotfyUnregisterFunction;
+	ComponentHandlerFunction HandlerFunction;
+};
+
+/* detail */ }
+
+template <
+	typename ComponentHandlerClass,
+	typename EOSHandle,
+	typename EOSNotfyRegisterFunction,
+	typename EOSNotfyUnregisterFunction,
+	typename ComponentHandlerFunction>
+EOSEventRegistrationPtr EOS_RegisterComponentEventHandler(
+	ComponentHandlerClass* HandlerClass,
+	EOSHandle ClientHandle,
+	int32_t ApiVersion,
+	EOSNotfyRegisterFunction NotfyRegisterFunction,
+	EOSNotfyUnregisterFunction NotfyUnregisterFunction,
+	ComponentHandlerFunction HandlerFunction)
+{
+	return MakeUnique<
+		detail::EOSEventRegistrationImpl<
+			ComponentHandlerClass,
+			EOSHandle,
+			EOSNotfyRegisterFunction,
+			EOSNotfyUnregisterFunction,
+			ComponentHandlerFunction>>(
+			HandlerClass,
+			ClientHandle,
+			ApiVersion,
+			NotfyRegisterFunction,
+			NotfyUnregisterFunction,
+			HandlerFunction);
+}
+
 /* UE::Online */ }
diff --git a/Engine/Source/Runtime/CoreOnline/Private/Online/CoreOnline.cpp b/Engine/Source/Runtime/CoreOnline/Private/Online/CoreOnline.cpp
--- a/Engine/Source/Runtime/CoreOnline/Private/Online/CoreOnline.cpp
+++ b/Engine/Source/Runtime/CoreOnline/Private/Online/CoreOnline.cpp
@@ -7,6 +7,193 @@
 
 namespace UE::Online {
 
+const TCHAR* LexToString(EOnlineServices Value)
+{
+	switch (Value)
+	{
+	case EOnlineServices::Null:				return TEXT("Null");
+	case EOnlineServices::Epic:				return TEXT("Epic");
+	case EOnlineServices::Xbox:				return TEXT("Xbox");
+	case EOnlineServices::PSN:				return TEXT("PSN");
+	case EOnlineServices::Nintendo:			return TEXT("Nintendo");
+	case EOnlineServices::Stadia:			return TEXT("Stadia");
+	case EOnlineServices::Steam:			return TEXT("Steam");
+	case EOnlineServices::Google:			return TEXT("Google");
+	case EOnlineServices::GooglePlay:		return TEXT("GooglePlay");
+	case EOnlineServices::Apple:			return TEXT("Apple");
+	case EOnlineServices::AppleGameKit:		return TEXT("AppleGameKit");
+	case EOnlineServices::Samsung:			return TEXT("Samsung");
+	case EOnlineServices::Oculus:			return TEXT("Oculus");
+	case EOnlineServices::Tencent:			return TEXT("Tencent");
+	case EOnlineServices::Reserved_14:		return TEXT("Reserved_14");
+	case EOnlineServices::Reserved_15:		return TEXT("Reserved_15");
+	case EOnlineServices::Reserved_16:		return TEXT("Reserved_16");
+	case EOnlineServices::Reserved_17:		return TEXT("Reserved_17");
+	case EOnlineServices::Reserved_18:		return TEXT("Reserved_18");
+	case EOnlineServices::Reserved_19:		return TEXT("Reserved_19");
+	case EOnlineServices::Reserved_20:		return TEXT("Reserved_20");
+	case EOnlineServices::Reserved_21:		return TEXT("Reserved_21");
+	case EOnlineServices::Reserved_22:		return TEXT("Reserved_22");
+	case EOnlineServices::Reserved_23:		return TEXT("Reserved_23");
+	case EOnlineServices::Reserved_24:		return TEXT("Reserved_24");
+	case EOnlineServices::Reserved_25:		return TEXT("Reserved_25");
+	case EOnlineServices::Reserved_26:		return TEXT("Reserved_26");
+	case EOnlineServices::Reserved_27:		return TEXT("Reserved_27");
+	case EOnlineServices::GameDefined_0:	return TEXT("GameDefined_0");
+	case EOnlineServices::GameDefined_1:	return TEXT("GameDefined_1");
+	case EOnlineServices::GameDefined_2:	return TEXT("GameDefined_2");
+	case EOnlineServices::GameDefined_3:	return TEXT("GameDefined_3");
+	case EOnlineServices::Platform:			return TEXT("Platform");
+	default:								checkNoEntry(); // Intentional fallthrough
+	case EOnlineServices::Default:			return TEXT("Default");
+	}
+}
+
+void LexFromString(EOnlineServices& OutValue, const TCHAR* InStr)
+{
+	if (FCString::Stricmp(InStr, TEXT("Null")) == 0)
+	{
+		OutValue = EOnlineServices::Null;
+	}
+	else if (FCString::Stricmp(InStr, TEXT("Epic")) == 0)
+	{
+		OutValue = EOnlineServices::Epic;
+	}
+	else if (FCString::Stricmp(InStr, TEXT("Xbox")) == 0)
+	{
+		OutValue = EOnlineServices::Xbox;
+	}
+	else if (FCString::Stricmp(InStr, TEXT("PSN")) == 0)
+	{
+		OutValue = EOnlineServices::PSN;
+	}
+	else if (FCString::Stricmp(InStr, TEXT("Nintendo")) == 0)
+	{
+		OutValue = EOnlineServices::Nintendo;
+	}
+	else if (FCString::Stricmp(InStr, TEXT("Stadia")) == 0)
+	{
+		OutValue = EOnlineServices::Stadia;
+	}
+	else if (FCString::Stricmp(InStr, TEXT("Steam")) == 0)
+	{
+		OutValue = EOnlineServices::Steam;
+	}
+	else if (FCString::Stricmp(InStr, TEXT("Google")) == 0)
+	{
+		OutValue = EOnlineServices::Google;
+	}
+	else if (FCString::Stricmp(InStr, TEXT("GooglePlay")) == 0)
+	{
+		OutValue = EOnlineServices::GooglePlay;
+	}
+	else if (FCString::Stricmp(InStr, TEXT("Apple")) == 0)
+	{
+		OutValue = EOnlineServices::Apple;
+	}
+	else if (FCString::Stricmp(InStr, TEXT("AppleGameKit")) == 0)
+	{
+		OutValue = EOnlineServices::AppleGameKit;
+	}
+	else if (FCString::Stricmp(InStr, TEXT("Samsung")) == 0)
+	{
+		OutValue = EOnlineServices::Samsung;
+	}
+	else if (FCString::Stricmp(InStr, TEXT("Oculus")) == 0)
+	{
+		OutValue = EOnlineServices::Oculus;
+	}
+	else if (FCString::Stricmp(InStr, TEXT("Tencent")) == 0)
+	{
+		OutValue = EOnlineServices::Tencent;
+	}
+	else if (FCString::Stricmp(InStr, TEXT("Reserved_14")) == 0)
+	{
+		OutValue = EOnlineServices::Reserved_14;
+	}
+	else if (FCString::Stricmp(InStr, TEXT("Reserved_15")) == 0)
+	{
+		OutValue = EOnlineServices::Reserved_15;
+	}
+	else if (FCString::Stricmp(InStr, TEXT("Reserved_16")) == 0)
+	{
+		OutValue = EOnlineServices::Reserved_16;
+	}
+	else if (FCString::Stricmp(InStr, TEXT("Reserved_17")) == 0)
+	{
+		OutValue = EOnlineServices::Reserved_17;
+	}
+	else if (FCString::Stricmp(InStr, TEXT("Reserved_18")) == 0)
+	{
+		OutValue = EOnlineServices::Reserved_18;
+	}
+	else if (FCString::Stricmp(InStr, TEXT("Reserved_19")) == 0)
+	{
+		OutValue = EOnlineServices::Reserved_19;
+	}
+	else if (FCString::Stricmp(InStr, TEXT("Reserved_20")) == 0)
+	{
+		OutValue = EOnlineServices::Reserved_20;
+	}
+	else if (FCString::Stricmp(InStr, TEXT("Reserved_21")) == 0)
+	{
+		OutValue = EOnlineServices::Reserved_21;
+	}
+	else if (FCString::Stricmp(InStr, TEXT("Reserved_22")) == 0)
+	{
+		OutValue = EOnlineServices::Reserved_22;
+	}
+	else if (FCString::Stricmp(InStr, TEXT("Reserved_23")) == 0)
+	{
+		OutValue = EOnlineServices::Reserved_23;
+	}
+	else if (FCString::Stricmp(InStr, TEXT("Reserved_24")) == 0)
+	{
+		OutValue = EOnlineServices::Reserved_24;
+	}
+	else if (FCString::Stricmp(InStr, TEXT("Reserved_25")) == 0)
+	{
+		OutValue = EOnlineServices::Reserved_25;
+	}
+	else if (FCString::Stricmp(InStr, TEXT("Reserved_26")) == 0)
+	{
+		OutValue = EOnlineServices::Reserved_26;
+	}
+	else if (FCString::Stricmp(InStr, TEXT("Reserved_27")) == 0)
+	{
+		OutValue = EOnlineServices::Reserved_27;
+	}
+	else if (FCString::Stricmp(InStr, TEXT("GameDefined_0")) == 0)
+	{
+		OutValue = EOnlineServices::GameDefined_0;
+	}
+	else if (FCString::Stricmp(InStr, TEXT("GameDefined_1")) == 0)
+	{
+		OutValue = EOnlineServices::GameDefined_1;
+	}
+	else if (FCString::Stricmp(InStr, TEXT("GameDefined_2")) == 0)
+	{
+		OutValue = EOnlineServices::GameDefined_2;
+	}
+	else if (FCString::Stricmp(InStr, TEXT("GameDefined_3")) == 0)
+	{
+		OutValue = EOnlineServices::GameDefined_3;
+	}
+	else if (FCString::Stricmp(InStr, TEXT("Platform")) == 0)
+	{
+		OutValue = EOnlineServices::Platform;
+	}
+	else if (FCString::Stricmp(InStr, TEXT("Default")) == 0)
+	{
+		OutValue = EOnlineServices::Default;
+	}
+	else
+	{
+		checkNoEntry();
+		OutValue = EOnlineServices::Default;
+	}
+}
+
 FOnlineIdRegistryRegistry& FOnlineIdRegistryRegistry::Get()
 {
 	return TLazySingleton<FOnlineIdRegistryRegistry>::Get();
@@ -101,6 +288,12 @@ FString ToLogString(const FOnlineAccountIdHandle& Id)
 	return FOnlineIdRegistryRegistry::Get().ToLogString(Id);
 }
 
+FString ToLogString(const FOnlineLobbyIdHandle& Id)
+{
+	// todo: use registry.
+	return FString::Printf(TEXT("%s:%d"), LexToString(Id.GetOnlineServicesType()), Id.GetHandle());
+}
+
 }	/* UE::Online */
 
 FString FUniqueNetIdWrapper::ToDebugString() const
diff --git a/Engine/Source/Runtime/CoreOnline/Public/Online/CoreOnline.h b/Engine/Source/Runtime/CoreOnline/Public/Online/CoreOnline.h
--- a/Engine/Source/Runtime/CoreOnline/Public/Online/CoreOnline.h
+++ b/Engine/Source/Runtime/CoreOnline/Public/Online/CoreOnline.h
@@ -197,7 +197,7 @@ namespace OnlineIdHandleTags
 {
 	struct FAccount {};
 	struct FSession {};
-	struct FParty {};
+	struct FLobby {};
 }
 
 enum class EOnlineServices : uint8
@@ -256,6 +256,9 @@ enum class EOnlineServices : uint8
 	Default = 255
 };
 
+COREONLINE_API const TCHAR* LexToString(EOnlineServices Value);
+COREONLINE_API void LexFromString(EOnlineServices& OutValue, const TCHAR* InStr);
+
 /**
  * A handle to an id which uniquely identifies a persistent or transient online resource, i.e. account/session/party etc, within a given Online Services provider.
  * At most one id, and therefore one handle, exists for any given resource. The id and handle persist until the OnlineServices module is unloaded.
@@ -285,8 +288,10 @@ private:
 };
 
 using FOnlineAccountIdHandle = TOnlineIdHandle<OnlineIdHandleTags::FAccount>;
+using FOnlineLobbyIdHandle = TOnlineIdHandle<OnlineIdHandleTags::FLobby>;
 
 COREONLINE_API FString ToLogString(const FOnlineAccountIdHandle& Id);
+COREONLINE_API FString ToLogString(const FOnlineLobbyIdHandle& Id);
 
 template<typename IdType>
 inline uint32 GetTypeHash(const TOnlineIdHandle<IdType>& Handle)
