Pass OnlineServices account id replication data by reference

diff --git a/Engine/Plugins/Online/Experimental/OnlineServicesEOS/Source/Private/Online/OnlineIdEOS.cpp b/Engine/Plugins/Online/Experimental/OnlineServicesEOS/Source/Private/Online/OnlineIdEOS.cpp
--- a/Engine/Plugins/Online/Experimental/OnlineServicesEOS/Source/Private/Online/OnlineIdEOS.cpp
+++ b/Engine/Plugins/Online/Experimental/OnlineServicesEOS/Source/Private/Online/OnlineIdEOS.cpp
@@ -226,7 +226,7 @@ TArray<uint8> FOnlineAccountIdRegistryEOS::ToReplicationData(const FOnlineAccoun
 	return ReplicationData;
 }
 
-FOnlineAccountIdHandle FOnlineAccountIdRegistryEOS::FromReplicationData(const TArray<uint8> ReplicationData)
+FOnlineAccountIdHandle FOnlineAccountIdRegistryEOS::FromReplicationData(const TArray<uint8>& ReplicationData)
 {
 	const EEOSAccountIdElements Elements = EEOSAccountIdElements(ReplicationData[0]);
 	const int EasHexBufferLength = EnumHasAnyFlags(Elements, EEOSAccountIdElements::EAS) ? OnlineIdEOSHexBufferLength : 0;
diff --git a/Engine/Plugins/Online/Experimental/OnlineServicesEOS/Source/Private/Online/OnlineIdEOS.h b/Engine/Plugins/Online/Experimental/OnlineServicesEOS/Source/Private/Online/OnlineIdEOS.h
--- a/Engine/Plugins/Online/Experimental/OnlineServicesEOS/Source/Private/Online/OnlineIdEOS.h
+++ b/Engine/Plugins/Online/Experimental/OnlineServicesEOS/Source/Private/Online/OnlineIdEOS.h
@@ -40,7 +40,7 @@ public:
 	// Begin IOnlineAccountIdRegistry
 	virtual FString ToLogString(const FOnlineAccountIdHandle& Handle) const override;
 	virtual TArray<uint8> ToReplicationData(const FOnlineAccountIdHandle& Handle) const override;
-	virtual FOnlineAccountIdHandle FromReplicationData(const TArray<uint8> ReplicationString) override;
+	virtual FOnlineAccountIdHandle FromReplicationData(const TArray<uint8>& ReplicationString) override;
 	// End IOnlineAccountIdRegistry
 
 	virtual ~FOnlineAccountIdRegistryEOS() = default;
diff --git a/Engine/Source/Runtime/CoreOnline/Public/Online/CoreOnline.h b/Engine/Source/Runtime/CoreOnline/Public/Online/CoreOnline.h
--- a/Engine/Source/Runtime/CoreOnline/Public/Online/CoreOnline.h
+++ b/Engine/Source/Runtime/CoreOnline/Public/Online/CoreOnline.h
@@ -308,7 +308,7 @@ class IOnlineIdRegistry
 public:
 	virtual FString ToLogString(const TOnlineIdHandle<IdType>& Handle) const = 0;
 	virtual TArray<uint8> ToReplicationData(const TOnlineIdHandle<IdType>& Handle) const = 0;
-	virtual TOnlineIdHandle<IdType> FromReplicationData(const TArray<uint8> Handle) = 0;
+	virtual TOnlineIdHandle<IdType> FromReplicationData(const TArray<uint8>& Handle) = 0;
 };
 
 using IOnlineAccountIdRegistry = IOnlineIdRegistry<OnlineIdHandleTags::FAccount>;
@@ -349,8 +349,9 @@ public:
 	COREONLINE_API TArray<uint8> ToReplicationData(const FOnlineAccountIdHandle& Handle) const;
 	COREONLINE_API FOnlineAccountIdHandle ToAccountId(EOnlineServices Services, const TArray<uint8>& RepData) const;
 
+	COREONLINE_API IOnlineAccountIdRegistry* GetAccountIdRegistry(EOnlineServices OnlineServices) const;
+
 private:
-	IOnlineAccountIdRegistry* GetAccountIdRegistry(EOnlineServices OnlineServices) const;
 
 	struct FAccountIdRegistryAndPriority
 	{
