- Fix deserializing account id's from rep data always using Null services.

diff --git a/Engine/Source/Runtime/CoreOnline/Private/Online/CoreOnline.cpp b/Engine/Source/Runtime/CoreOnline/Private/Online/CoreOnline.cpp
--- a/Engine/Source/Runtime/CoreOnline/Private/Online/CoreOnline.cpp
+++ b/Engine/Source/Runtime/CoreOnline/Private/Online/CoreOnline.cpp
@@ -268,7 +268,7 @@ TArray<uint8> FOnlineIdRegistryRegistry::ToReplicationData(const FOnlineAccountI
 FOnlineAccountIdHandle FOnlineIdRegistryRegistry::ToAccountId(EOnlineServices Services, const TArray<uint8>& RepData) const
 {
 	FOnlineAccountIdHandle Handle;
-	if (IOnlineAccountIdRegistry* Registry = GetAccountIdRegistry(Handle.GetOnlineServicesType()))
+	if (IOnlineAccountIdRegistry* Registry = GetAccountIdRegistry(Services))
 	{
 		Handle = Registry->FromReplicationData(RepData);
 	}
