OSSv2:
Rename FOnlineSessionInviteIdHandle to FSessionInviteId;
Rename FOnlineVerifiedAuthTicketIdHandle to FVerifiedAuthTicketId;
Rename FOnlineVerifiedAuthSessionIdHandle to FVerifiedAuthSessionId;
Rename FOnlineIdHandle to FOnlineId and corresponding variable name Handle to OnlineId;

diff --git a/Engine/Plugins/Online/Experimental/OnlineServices/Source/OnlineServicesCommon/Private/Online/SessionsCommon.cpp b/Engine/Plugins/Online/Experimental/OnlineServices/Source/OnlineServicesCommon/Private/Online/SessionsCommon.cpp
--- a/Engine/Plugins/Online/Experimental/OnlineServices/Source/OnlineServicesCommon/Private/Online/SessionsCommon.cpp
+++ b/Engine/Plugins/Online/Experimental/OnlineServices/Source/OnlineServicesCommon/Private/Online/SessionsCommon.cpp
@@ -225,9 +225,9 @@ namespace UE::Online {
 	{
 		TArray<TSharedRef<const FSessionInvite>> SessionInvites;
 
-		if (const TMap<FOnlineSessionInviteIdHandle, TSharedRef<FSessionInvite>>* UserMap = SessionInvitesUserMap.Find(Params.LocalAccountId))
+		if (const TMap<FSessionInviteId, TSharedRef<FSessionInvite>>* UserMap = SessionInvitesUserMap.Find(Params.LocalAccountId))
 		{
-			for (const TPair<FOnlineSessionInviteIdHandle, TSharedRef<FSessionInvite>>& Entry : *UserMap)
+			for (const TPair<FSessionInviteId, TSharedRef<FSessionInvite>>& Entry : *UserMap)
 			{
 				SessionInvites.Add(Entry.Value);
 			}
@@ -364,7 +364,7 @@ namespace UE::Online {
 	{
 		AllSessionsById.Emplace(Session->SessionId, Session);
 
-		TMap<FOnlineSessionInviteIdHandle, TSharedRef<FSessionInvite>>& SessionInvitesMap = SessionInvitesUserMap.FindOrAdd(LocalAccountId);
+		TMap<FSessionInviteId, TSharedRef<FSessionInvite>>& SessionInvitesMap = SessionInvitesUserMap.FindOrAdd(LocalAccountId);
 		SessionInvitesMap.Emplace(SessionInvite->InviteId, SessionInvite);
 	}
 
@@ -397,10 +397,10 @@ namespace UE::Online {
 
 	void FSessionsCommon::ClearSessionInvitesForSession(const FAccountId& LocalAccountId, const FOnlineSessionId SessionId)
 	{
-		if (TMap<FOnlineSessionInviteIdHandle, TSharedRef<FSessionInvite>>* UserMap = SessionInvitesUserMap.Find(LocalAccountId))
+		if (TMap<FSessionInviteId, TSharedRef<FSessionInvite>>* UserMap = SessionInvitesUserMap.Find(LocalAccountId))
 		{
-			TArray<FOnlineSessionInviteIdHandle> InviteIdsToRemove;
-			for (const TPair<FOnlineSessionInviteIdHandle, TSharedRef<FSessionInvite>>& Entry : *UserMap)
+			TArray<FSessionInviteId> InviteIdsToRemove;
+			for (const TPair<FSessionInviteId, TSharedRef<FSessionInvite>>& Entry : *UserMap)
 			{
 				if (Entry.Value->SessionId == SessionId)
 				{
@@ -408,7 +408,7 @@ namespace UE::Online {
 				}
 			}
 
-			for (const FOnlineSessionInviteIdHandle& InviteId : InviteIdsToRemove)
+			for (const FSessionInviteId& InviteId : InviteIdsToRemove)
 			{
 				UserMap->Remove(InviteId);
 			}
@@ -453,9 +453,9 @@ namespace UE::Online {
 			}
 		}
 
-		for (const TPair<FAccountId, TMap<FOnlineSessionInviteIdHandle, TSharedRef<FSessionInvite>>>& Entry : SessionInvitesUserMap)
+		for (const TPair<FAccountId, TMap<FSessionInviteId, TSharedRef<FSessionInvite>>>& Entry : SessionInvitesUserMap)
 		{
-			for (const TPair<FOnlineSessionInviteIdHandle, TSharedRef<FSessionInvite>>& InviteMap : Entry.Value)
+			for (const TPair<FSessionInviteId, TSharedRef<FSessionInvite>>& InviteMap : Entry.Value)
 			{
 				if (InviteMap.Value->SessionId == SessionId)
 				{
diff --git a/Engine/Plugins/Online/Experimental/OnlineServices/Source/OnlineServicesCommon/Public/Online/OnlineExecHandler.h b/Engine/Plugins/Online/Experimental/OnlineServices/Source/OnlineServicesCommon/Public/Online/OnlineExecHandler.h
--- a/Engine/Plugins/Online/Experimental/OnlineServices/Source/OnlineServicesCommon/Public/Online/OnlineExecHandler.h
+++ b/Engine/Plugins/Online/Experimental/OnlineServices/Source/OnlineServicesCommon/Public/Online/OnlineExecHandler.h
@@ -122,7 +122,7 @@ template <typename T> inline bool ParseOnlineExecParams(const TCHAR*& Cmd, TShar
 template <typename T> inline bool ParseOnlineExecParams(const TCHAR*& Cmd, TSharedRef<T>& Ref, IOnlineServices* Services = nullptr);
 template <typename T> inline bool ParseOnlineExecParams(const TCHAR*& Cmd, TOptional<T>& Optional, IOnlineServices* Services = nullptr);
 template <typename... Ts> inline bool ParseOnlineExecParams(const TCHAR*& Cmd, TVariant<Ts...>& Variant, IOnlineServices* Services = nullptr);
-template <typename IdType> inline bool ParseOnlineExecParams(const TCHAR*& Cmd, TOnlineIdHandle<IdType>& Value, IOnlineServices* Services = nullptr);
+template <typename IdType> inline bool ParseOnlineExecParams(const TCHAR*& Cmd, TOnlineId<IdType>& Value, IOnlineServices* Services = nullptr);
 template <typename T> std::enable_if_t<!TModels<Meta::COnlineMetadataAvailable, T>::Value, bool> ParseOnlineExecParams(const TCHAR*& Cmd, T& Value, IOnlineServices* Services = nullptr);
 template <typename T> std::enable_if_t<TModels<Meta::COnlineMetadataAvailable, T>::Value, bool> ParseOnlineExecParams(const TCHAR*& Cmd, T& Value, IOnlineServices* Services = nullptr);
 
@@ -231,7 +231,7 @@ public:
 	{
 		FAccountId AccountId;
 		const TCHAR* ValueTCHAR = *VariantValue;
-		if (ParseOnlineExecParams<TOnlineIdHandle<OnlineIdHandleTags::FAccount>>(ValueTCHAR, AccountId, Services))
+		if (ParseOnlineExecParams<TOnlineId<OnlineIdHandleTags::FAccount>>(ValueTCHAR, AccountId, Services))
 		{
 			Variant.template Set<FAccountId>(AccountId);
 			return true;
@@ -762,7 +762,7 @@ inline bool ParseOnlineExecParams(const TCHAR*& Cmd, TVariant<Ts...>& Variant, I
 }
 
 template<typename IdType>
-inline bool ParseOnlineExecParams(const TCHAR*& Cmd, TOnlineIdHandle<IdType>& Value, IOnlineServices* Services)
+inline bool ParseOnlineExecParams(const TCHAR*& Cmd, TOnlineId<IdType>& Value, IOnlineServices* Services)
 {
 	FString Token;
 	if (FParse::Token(Cmd, Token, true))
@@ -775,7 +775,7 @@ inline bool ParseOnlineExecParams(const TCHAR*& Cmd, TOnlineIdHandle<IdType>& Va
 			EOnlineServices ServiceEnum;
 			LexFromString(ServiceEnum, *ServicesString);
 			uint32 Handle = static_cast<uint32>(FCString::Strtoui64(*HandleString, nullptr, 10));
-			Value = TOnlineIdHandle<IdType>(ServiceEnum, Handle);
+			Value = TOnlineId<IdType>(ServiceEnum, Handle);
 			return true;
 		}
 	}
@@ -784,7 +784,7 @@ inline bool ParseOnlineExecParams(const TCHAR*& Cmd, TOnlineIdHandle<IdType>& Va
 }
 
 template<>
-inline bool ParseOnlineExecParams<OnlineIdHandleTags::FAccount>(const TCHAR*& Cmd, TOnlineIdHandle<OnlineIdHandleTags::FAccount>& Value, IOnlineServices* Services)
+inline bool ParseOnlineExecParams<OnlineIdHandleTags::FAccount>(const TCHAR*& Cmd, TOnlineId<OnlineIdHandleTags::FAccount>& Value, IOnlineServices* Services)
 {
 	FString Token;
 	if (FParse::Token(Cmd, Token, true))
@@ -797,7 +797,7 @@ inline bool ParseOnlineExecParams<OnlineIdHandleTags::FAccount>(const TCHAR*& Cm
 			EOnlineServices ServicesEnum;
 			LexFromString(ServicesEnum, *ServicesString);
 			uint32 Handle = static_cast<uint32>(FCString::Strtoui64(*HandleString, nullptr, 10));
-			Value = TOnlineIdHandle<OnlineIdHandleTags::FAccount>(ServicesEnum, Handle);
+			Value = TOnlineId<OnlineIdHandleTags::FAccount>(ServicesEnum, Handle);
 			return true;
 		}
 		else if (Token.Len() == 1 && **Token >= TEXT('0') && **Token <= TEXT('9'))
diff --git a/Engine/Plugins/Online/Experimental/OnlineServices/Source/OnlineServicesCommon/Public/Online/OnlineIdCommon.h b/Engine/Plugins/Online/Experimental/OnlineServices/Source/OnlineServicesCommon/Public/Online/OnlineIdCommon.h
--- a/Engine/Plugins/Online/Experimental/OnlineServices/Source/OnlineServicesCommon/Public/Online/OnlineIdCommon.h
+++ b/Engine/Plugins/Online/Experimental/OnlineServices/Source/OnlineServicesCommon/Public/Online/OnlineIdCommon.h
@@ -21,7 +21,7 @@ public:
 
 	virtual ~TOnlineBasicIdRegistry() = default;
 
-	using HandleType = TOnlineIdHandle<IdType>;
+	using HandleType = TOnlineId<IdType>;
 
 	HandleType GetInvalidHandle() const
 	{
diff --git a/Engine/Plugins/Online/Experimental/OnlineServices/Source/OnlineServicesCommon/Public/Online/SessionsCommon.h b/Engine/Plugins/Online/Experimental/OnlineServices/Source/OnlineServicesCommon/Public/Online/SessionsCommon.h
--- a/Engine/Plugins/Online/Experimental/OnlineServices/Source/OnlineServicesCommon/Public/Online/SessionsCommon.h
+++ b/Engine/Plugins/Online/Experimental/OnlineServices/Source/OnlineServicesCommon/Public/Online/SessionsCommon.h
@@ -76,9 +76,9 @@ class FOnlineSessionInviteIdStringRegistry : public IOnlineSessionInviteIdRegist
 {
 public:
 	// Begin IOnlineSessionIdRegistry
-	virtual inline FString ToLogString(const FOnlineSessionInviteIdHandle& Handle) const override
+	virtual inline FString ToLogString(const FSessionInviteId& SessionInviteId) const override
 	{
-		FString IdValue = BasicRegistry.FindIdValue(Handle);
+		FString IdValue = BasicRegistry.FindIdValue(SessionInviteId);
 
 		if (IdValue.Len() == 0)
 		{
@@ -88,9 +88,9 @@ public:
 		return IdValue;
 	};
 
-	virtual inline TArray<uint8> ToReplicationData(const FOnlineSessionInviteIdHandle& Handle) const override
+	virtual inline TArray<uint8> ToReplicationData(const FSessionInviteId& SessionInviteId) const override
 	{
-		const FString IdValue = BasicRegistry.FindIdValue(Handle);
+		const FString IdValue = BasicRegistry.FindIdValue(SessionInviteId);
 		const FTCHARToUTF8 IdValueUtf8(IdValue);
 
 		TArray<uint8> ReplicationData;
@@ -101,7 +101,7 @@ public:
 		return ReplicationData;
 	}
 
-	virtual inline FOnlineSessionInviteIdHandle FromReplicationData(const TArray<uint8>& ReplicationData) override
+	virtual inline FSessionInviteId FromReplicationData(const TArray<uint8>& ReplicationData) override
 	{
 		const FUTF8ToTCHAR IdValueTCHAR((char*)ReplicationData.GetData(), ReplicationData.Num());
 		const FString IdValue = FString(IdValueTCHAR.Length(), IdValueTCHAR.Get());
@@ -111,7 +111,7 @@ public:
 			return BasicRegistry.FindOrAddHandle(IdValue);
 		}
 
-		return FOnlineSessionInviteIdHandle();
+		return FSessionInviteId();
 	}
 	// End IOnlineSessionIdRegistry
 
@@ -297,7 +297,7 @@ protected:
 	TMap<FAccountId, FOnlineSessionId> PresenceSessionsUserMap;
 
 	/** Cache for received session invites, mapped per user */
-	TMap<FAccountId, TMap<FOnlineSessionInviteIdHandle, TSharedRef<FSessionInvite>>> SessionInvitesUserMap;
+	TMap<FAccountId, TMap<FSessionInviteId, TSharedRef<FSessionInvite>>> SessionInvitesUserMap;
 
 	/** Cache for the last set of session search results, mapped per user */
 	TMap<FAccountId, TArray<FOnlineSessionId>> SearchResultsUserMap;
diff --git a/Engine/Plugins/Online/Experimental/OnlineServices/Source/OnlineServicesInterface/Public/Online/Auth.h b/Engine/Plugins/Online/Experimental/OnlineServices/Source/OnlineServicesInterface/Public/Online/Auth.h
--- a/Engine/Plugins/Online/Experimental/OnlineServices/Source/OnlineServicesInterface/Public/Online/Auth.h
+++ b/Engine/Plugins/Online/Experimental/OnlineServices/Source/OnlineServicesInterface/Public/Online/Auth.h
@@ -141,7 +141,7 @@ using FCredentialsToken = TVariant<FString, FExternalAuthToken>;
 
 struct FVerifiedAuthSession
 {
-	FOnlineVerifiedAuthSessionIdHandle SessionId;
+	FVerifiedAuthSessionId SessionId;
 	FAccountId RemoteAccountId;
 	double CreationTime;
 };
@@ -260,7 +260,7 @@ struct FAuthQueryVerifiedAuthTicket
 	struct Result
 	{
 		/** Local ticket id used to reference the ticket in further operations. */
-		FOnlineVerifiedAuthTicketIdHandle VerifiedAuthTicketId;
+		FVerifiedAuthTicketId VerifiedAuthTicketId;
 		/** Ticket used to begin a verified auth session with a remote host. */
 		FVerifiedAuthTicket VerifiedAuthTicket;
 	};
@@ -275,7 +275,7 @@ struct FAuthCancelVerifiedAuthTicket
 		/** The online account id of the Local User making the request. */
 		FAccountId LocalAccountId;
 		/** Local ticket id used to reference the verified auth ticket. */
-		FOnlineVerifiedAuthTicketIdHandle VerifiedAuthTicketId;
+		FVerifiedAuthTicketId VerifiedAuthTicketId;
 	};
 
 	struct Result
@@ -307,7 +307,7 @@ struct FAuthEndVerifiedAuthSession
 
 	struct Params
 	{
-		FOnlineVerifiedAuthSessionIdHandle SessionId;
+		FVerifiedAuthSessionId SessionId;
 	};
 
 	struct Result
diff --git a/Engine/Plugins/Online/Experimental/OnlineServices/Source/OnlineServicesInterface/Public/Online/Sessions.h b/Engine/Plugins/Online/Experimental/OnlineServices/Source/OnlineServicesInterface/Public/Online/Sessions.h
--- a/Engine/Plugins/Online/Experimental/OnlineServices/Source/OnlineServicesInterface/Public/Online/Sessions.h
+++ b/Engine/Plugins/Online/Experimental/OnlineServices/Source/OnlineServicesInterface/Public/Online/Sessions.h
@@ -197,7 +197,7 @@ struct FSessionInvite
 	FAccountId SenderId;
 
 	/* The invite id handle, needed for retrieving session information and rejecting the invite */
-	FOnlineSessionInviteIdHandle InviteId;
+	FSessionInviteId InviteId;
 
 	/* Pointer to the session information */
 	FOnlineSessionId SessionId;
@@ -550,7 +550,7 @@ struct FRejectSessionInvite
 		FAccountId LocalAccountId;
 
 		/* The id handle for the invite to be rejected */
-		FOnlineSessionInviteIdHandle SessionInviteId;
+		FSessionInviteId SessionInviteId;
 	};
 
 	struct Result
diff --git a/Engine/Plugins/Online/Experimental/OnlineServicesEOSGS/Source/Private/Online/SessionsEOSGS.cpp b/Engine/Plugins/Online/Experimental/OnlineServicesEOSGS/Source/Private/Online/SessionsEOSGS.cpp
--- a/Engine/Plugins/Online/Experimental/OnlineServicesEOSGS/Source/Private/Online/SessionsEOSGS.cpp
+++ b/Engine/Plugins/Online/Experimental/OnlineServicesEOSGS/Source/Private/Online/SessionsEOSGS.cpp
@@ -289,14 +289,14 @@ void FSessionsEOSGS::HandleSessionInviteReceived(const EOS_Sessions_SessionInvit
 					const FAccountId& ReceiverId = Result.GetOkValue().LocalAccountId;
 					const FAccountId& SenderId = FindAccountId(SenderPUID);
 
-					FOnlineSessionInviteIdHandle InviteIdHandle = CreateSessionInviteId(InviteId);
+					FSessionInviteId SessionInviteId = CreateSessionInviteId(InviteId);
 
 					TSharedRef<FSessionCommon> Session = Result.GetOkValue().Session;
 
 					TSharedRef<FSessionInvite> SessionInviteRef = MakeShared<FSessionInvite>(FSessionInvite{
 						ReceiverId,
 						SenderId,
-						InviteIdHandle,
+						SessionInviteId,
 						Session->SessionId
 						});
 
@@ -1243,7 +1243,7 @@ FOnlineSessionId FSessionsEOSGS::CreateSessionId(const FString& SessionId)
 	return FOnlineSessionIdRegistryEOSGS::Get().BasicRegistry.FindOrAddHandle(SessionId);
 }
 
-FOnlineSessionInviteIdHandle FSessionsEOSGS::CreateSessionInviteId(const FString& SessionInviteId) const
+FSessionInviteId FSessionsEOSGS::CreateSessionInviteId(const FString& SessionInviteId) const
 {
 	return FOnlineSessionInviteIdRegistryEOSGS::Get().BasicRegistry.FindOrAddHandle(SessionInviteId);
 }
@@ -1784,7 +1784,7 @@ TOnlineAsyncOpHandle<FRejectSessionInvite> FSessionsEOSGS::RejectSessionInvite(F
 
 		const FRejectSessionInvite::Params& OpParams = Op.GetParams();
 
-		if (TMap<FOnlineSessionInviteIdHandle, TSharedRef<FSessionInvite>>* UserMap = SessionInvitesUserMap.Find(OpParams.LocalAccountId))
+		if (TMap<FSessionInviteId, TSharedRef<FSessionInvite>>* UserMap = SessionInvitesUserMap.Find(OpParams.LocalAccountId))
 		{
 			UserMap->Remove(OpParams.SessionInviteId);
 		}
diff --git a/Engine/Plugins/Online/Experimental/OnlineServicesEOSGS/Source/Public/Online/OnlineIdEOSGS.h b/Engine/Plugins/Online/Experimental/OnlineServicesEOSGS/Source/Public/Online/OnlineIdEOSGS.h
--- a/Engine/Plugins/Online/Experimental/OnlineServicesEOSGS/Source/Public/Online/OnlineIdEOSGS.h
+++ b/Engine/Plugins/Online/Experimental/OnlineServicesEOSGS/Source/Public/Online/OnlineIdEOSGS.h
@@ -58,9 +58,9 @@ FAccountId ONLINESERVICESEOSGS_API FindAccountId(const EOS_ProductUserId EpicAcc
 FAccountId ONLINESERVICESEOSGS_API FindAccountIdChecked(const EOS_ProductUserId EpicAccountId);
 
 template<typename IdType>
-inline bool ValidateOnlineId(const TOnlineIdHandle<IdType> Handle)
+inline bool ValidateOnlineId(const TOnlineId<IdType> OnlineId)
 {
-	return Handle.GetOnlineServicesType() == EOnlineServices::Epic && Handle.IsValid();
+	return OnlineId.GetOnlineServicesType() == EOnlineServices::Epic && OnlineId.IsValid();
 }
 
 } /* namespace UE::Online */
diff --git a/Engine/Plugins/Online/Experimental/OnlineServicesEOSGS/Source/Public/Online/SessionsEOSGS.h b/Engine/Plugins/Online/Experimental/OnlineServicesEOSGS/Source/Public/Online/SessionsEOSGS.h
--- a/Engine/Plugins/Online/Experimental/OnlineServicesEOSGS/Source/Public/Online/SessionsEOSGS.h
+++ b/Engine/Plugins/Online/Experimental/OnlineServicesEOSGS/Source/Public/Online/SessionsEOSGS.h
@@ -251,7 +251,7 @@ protected:
 	TOnlineAsyncOpHandle<FSendSingleSessionInviteImpl> SendSingleSessionInviteImpl(FSendSingleSessionInviteImpl::Params&& Params);
 
 	static FOnlineSessionId CreateSessionId(const FString& SessionId);
-	FOnlineSessionInviteIdHandle CreateSessionInviteId(const FString& SessionInviteId) const;
+	FSessionInviteId CreateSessionInviteId(const FString& SessionInviteId) const;
 
 	/**
 	 * Builds a session from an EOS Session Details Handle. Asynchronous due to the id resolution process
diff --git a/Engine/Plugins/Online/Experimental/OnlineServicesOSSAdapter/Source/Public/Online/OnlineIdOSSAdapter.h b/Engine/Plugins/Online/Experimental/OnlineServicesOSSAdapter/Source/Public/Online/OnlineIdOSSAdapter.h
--- a/Engine/Plugins/Online/Experimental/OnlineServicesOSSAdapter/Source/Public/Online/OnlineIdOSSAdapter.h
+++ b/Engine/Plugins/Online/Experimental/OnlineServicesOSSAdapter/Source/Public/Online/OnlineIdOSSAdapter.h
@@ -22,12 +22,12 @@ public:
 
 	virtual ~TOnlineUniqueNetIdRegistry() {}
 
-	TOnlineIdHandle<IdType> FindOrAddHandle(const FUniqueNetIdRef& IdValue)
+	TOnlineId<IdType> FindOrAddHandle(const FUniqueNetIdRef& IdValue)
 	{
-		TOnlineIdHandle<IdType> Handle;
+		TOnlineId<IdType> OnlineId;
 		if (!ensure(IdValue->IsValid()))
 		{
-			return Handle;
+			return OnlineId;
 		}
 
 		// Take a read lock and check if we already have a handle
@@ -35,47 +35,47 @@ public:
 			FReadScopeLock ReadLock(Lock);
 			if (const uint32* FoundHandle = IdValueToHandleMap.Find(IdValue))
 			{
-				Handle = TOnlineIdHandle<IdType>(OnlineServicesType, *FoundHandle);
+				OnlineId = TOnlineId<IdType>(OnlineServicesType, *FoundHandle);
 			}
 		}
 
-		if (!Handle.IsValid())
+		if (!OnlineId.IsValid())
 		{
 			// Take a write lock, check again if we already have a handle, or insert a new element.
 			FWriteScopeLock WriteLock(Lock);
 			if (const uint32* FoundHandle = IdValueToHandleMap.Find(IdValue))
 			{
-				Handle = TOnlineIdHandle<IdType>(OnlineServicesType, *FoundHandle);
+				OnlineId = TOnlineId<IdType>(OnlineServicesType, *FoundHandle);
 			}
 
-			if (!Handle.IsValid())
+			if (!OnlineId.IsValid())
 			{
 				IdValues.Emplace(IdValue);
-				Handle = TOnlineIdHandle<IdType>(OnlineServicesType, IdValues.Num());
-				IdValueToHandleMap.Emplace(IdValue, Handle.GetHandle());
+				OnlineId = TOnlineId<IdType>(OnlineServicesType, IdValues.Num());
+				IdValueToHandleMap.Emplace(IdValue, OnlineId.GetHandle());
 			}
 		}
 
-		return Handle;
+		return OnlineId;
 	}
 
 	// Returns a copy as it's not thread safe to return a pointer/ref to an element of an array that can be relocated by another thread.
-	FUniqueNetIdPtr GetIdValue(const TOnlineIdHandle<IdType> Handle) const
+	FUniqueNetIdPtr GetIdValue(const TOnlineId<IdType> OnlineId) const
 	{
-		if (Handle.GetOnlineServicesType() == OnlineServicesType && Handle.IsValid())
+		if (OnlineId.GetOnlineServicesType() == OnlineServicesType && OnlineId.IsValid())
 		{
 			FReadScopeLock ReadLock(Lock);
-			if (IdValues.IsValidIndex(Handle.GetHandle() - 1))
+			if (IdValues.IsValidIndex(OnlineId.GetHandle() - 1))
 			{
-				return IdValues[Handle.GetHandle() - 1];
+				return IdValues[OnlineId.GetHandle() - 1];
 			}
 		}
 		return FUniqueNetIdPtr();
 	}
 
-	FUniqueNetIdRef GetIdValueChecked(const TOnlineIdHandle<IdType> Handle) const
+	FUniqueNetIdRef GetIdValueChecked(const TOnlineId<IdType> OnlineId) const
 	{
-		return GetIdValue(Handle).ToSharedRef();
+		return GetIdValue(OnlineId).ToSharedRef();
 	}
 
 	bool IsHandleExpired(const FOnlineSessionId& InSessionId) const
@@ -84,20 +84,20 @@ public:
 	}
 
 	// Begin IOnlineAccountIdRegistry
-	virtual FString ToLogString(const TOnlineIdHandle<IdType>& Handle) const override
+	virtual FString ToLogString(const TOnlineId<IdType>& OnlineId) const override
 	{
-		FUniqueNetIdPtr IdValue = GetIdValue(Handle);
+		FUniqueNetIdPtr IdValue = GetIdValue(OnlineId);
 		return IdValue ? IdValue->ToDebugString() : FString(TEXT("invalid_id"));
 	}
 
-	virtual TArray<uint8> ToReplicationData(const TOnlineIdHandle<IdType>& Handle) const override
+	virtual TArray<uint8> ToReplicationData(const TOnlineId<IdType>& OnlineId) const override
 	{
 		return TArray<uint8>();
 	}
 
-	virtual TOnlineIdHandle<IdType> FromReplicationData(const TArray<uint8>& Handle) override
+	virtual TOnlineId<IdType> FromReplicationData(const TArray<uint8>& OnlineId) override
 	{
-		return TOnlineIdHandle<IdType>();
+		return TOnlineId<IdType>();
 	}
 
 	// End IOnlineAccountIdRegistry
diff --git a/Engine/Source/Runtime/CoreOnline/Private/Online/CoreOnline.cpp b/Engine/Source/Runtime/CoreOnline/Private/Online/CoreOnline.cpp
--- a/Engine/Source/Runtime/CoreOnline/Private/Online/CoreOnline.cpp
+++ b/Engine/Source/Runtime/CoreOnline/Private/Online/CoreOnline.cpp
@@ -365,34 +365,34 @@ void FOnlineIdRegistryRegistry::UnregisterSessionInviteIdRegistry(EOnlineService
 	}
 }
 
-FString FOnlineIdRegistryRegistry::ToLogString(const FOnlineSessionInviteIdHandle& Handle) const
+FString FOnlineIdRegistryRegistry::ToLogString(const FSessionInviteId& SessionInviteId) const
 {
 	FString Result;
-	if (IOnlineSessionInviteIdRegistry* Registry = GetSessionInviteIdRegistry(Handle.GetOnlineServicesType()))
+	if (IOnlineSessionInviteIdRegistry* Registry = GetSessionInviteIdRegistry(SessionInviteId.GetOnlineServicesType()))
 	{
-		Result = Registry->ToLogString(Handle);
+		Result = Registry->ToLogString(SessionInviteId);
 	}
 	return Result;
 }
 
-TArray<uint8> FOnlineIdRegistryRegistry::ToReplicationData(const FOnlineSessionInviteIdHandle& Handle) const
+TArray<uint8> FOnlineIdRegistryRegistry::ToReplicationData(const FSessionInviteId& SessionInviteId) const
 {
 	TArray<uint8> RepData;
-	if (IOnlineSessionInviteIdRegistry* Registry = GetSessionInviteIdRegistry(Handle.GetOnlineServicesType()))
+	if (IOnlineSessionInviteIdRegistry* Registry = GetSessionInviteIdRegistry(SessionInviteId.GetOnlineServicesType()))
 	{
-		RepData = Registry->ToReplicationData(Handle);
+		RepData = Registry->ToReplicationData(SessionInviteId);
 	}
 	return RepData;
 }
 
-FOnlineSessionInviteIdHandle FOnlineIdRegistryRegistry::ToSessionInviteId(EOnlineServices Services, const TArray<uint8>& RepData) const
+FSessionInviteId FOnlineIdRegistryRegistry::ToSessionInviteId(EOnlineServices Services, const TArray<uint8>& RepData) const
 {
-	FOnlineSessionInviteIdHandle Handle;
+	FSessionInviteId SessionInviteId;
 	if (IOnlineSessionInviteIdRegistry* Registry = GetSessionInviteIdRegistry(Services))
 	{
-		Handle = Registry->FromReplicationData(RepData);
+		SessionInviteId = Registry->FromReplicationData(RepData);
 	}
-	return Handle;
+	return SessionInviteId;
 }
 
 IOnlineSessionInviteIdRegistry* FOnlineIdRegistryRegistry::GetSessionInviteIdRegistry(EOnlineServices OnlineServices) const
@@ -421,17 +421,17 @@ FString ToLogString(const FOnlineSessionId& Id)
 	return FString::Printf(TEXT("%s:%d (%s)"), LexToString(Id.GetOnlineServicesType()), Id.GetHandle(), *FOnlineIdRegistryRegistry::Get().ToLogString(Id));
 }
 
-FString ToLogString(const FOnlineSessionInviteIdHandle& Id)
+FString ToLogString(const FSessionInviteId& Id)
 {
 	return FString::Printf(TEXT("%s:%d (%s)"), LexToString(Id.GetOnlineServicesType()), Id.GetHandle(), *FOnlineIdRegistryRegistry::Get().ToLogString(Id));
 }
 
-FString ToLogString(const FOnlineVerifiedAuthTicketIdHandle& Id)
+FString ToLogString(const FVerifiedAuthTicketId& Id)
 {
 	return FString::Printf(TEXT("%s:%d"), LexToString(Id.GetOnlineServicesType()), Id.GetHandle());
 }
 
-FString ToLogString(const FOnlineVerifiedAuthSessionIdHandle& Id)
+FString ToLogString(const FVerifiedAuthSessionId& Id)
 {
 	return FString::Printf(TEXT("%s:%d"), LexToString(Id.GetOnlineServicesType()), Id.GetHandle());
 }
diff --git a/Engine/Source/Runtime/CoreOnline/Public/Online/CoreOnline.h b/Engine/Source/Runtime/CoreOnline/Public/Online/CoreOnline.h
--- a/Engine/Source/Runtime/CoreOnline/Public/Online/CoreOnline.h
+++ b/Engine/Source/Runtime/CoreOnline/Public/Online/CoreOnline.h
@@ -206,7 +206,7 @@ namespace UE::Online {
 
 class FOnlineForeignAccountIdRegistry;
 
-/** Tags used as template argument to TOnlineIdHandle to make it a compile error to assign between id's of different types */
+/** Tags used as template argument to TOnlineId to make it a compile error to assign between id's of different types */
 namespace OnlineIdHandleTags
 {
 	struct FAccount {};
@@ -284,11 +284,11 @@ COREONLINE_API void LexFromString(EOnlineServices& OutValue, const TCHAR* InStr)
  * Passed to and returned from OnlineServices APIs.
  */
 template<typename IdType>
-class TOnlineIdHandle
+class TOnlineId
 {
 public:
-	TOnlineIdHandle() = default;
-	TOnlineIdHandle(EOnlineServices Type, uint32 Handle)
+	TOnlineId() = default;
+	TOnlineId(EOnlineServices Type, uint32 Handle)
 	{
 		check(Handle < 0xFF000000);
 		Value = (Handle & 0x00FFFFFF) | (uint32(Type) << 24);
@@ -299,32 +299,32 @@ public:
 	EOnlineServices GetOnlineServicesType() const { return EOnlineServices(Value >> 24); }
 	uint32 GetHandle() const { return Value & 0x00FFFFFF; }
 
-	bool operator==(const TOnlineIdHandle& Other) const { return Value == Other.Value; }
-	bool operator!=(const TOnlineIdHandle& Other) const { return Value != Other.Value; }
+	bool operator==(const TOnlineId& Other) const { return Value == Other.Value; }
+	bool operator!=(const TOnlineId& Other) const { return Value != Other.Value; }
 
 private:
 	uint32 Value = uint32(EOnlineServices::Null) << 24;
 };
 
-using FAccountId = TOnlineIdHandle<OnlineIdHandleTags::FAccount>;
-using FLobbyId = TOnlineIdHandle<OnlineIdHandleTags::FLobby>;
-using FOnlineSessionId = TOnlineIdHandle<OnlineIdHandleTags::FSession>;
-using FOnlineSessionInviteIdHandle = TOnlineIdHandle<OnlineIdHandleTags::FSessionInvite>;
-using FOnlineVerifiedAuthTicketIdHandle = TOnlineIdHandle<OnlineIdHandleTags::FVerifiedAuthTicket>;
-using FOnlineVerifiedAuthSessionIdHandle = TOnlineIdHandle<OnlineIdHandleTags::FVerifiedAuthSession>;
+using FAccountId = TOnlineId<OnlineIdHandleTags::FAccount>;
+using FLobbyId = TOnlineId<OnlineIdHandleTags::FLobby>;
+using FOnlineSessionId = TOnlineId<OnlineIdHandleTags::FSession>;
+using FSessionInviteId = TOnlineId<OnlineIdHandleTags::FSessionInvite>;
+using FVerifiedAuthTicketId = TOnlineId<OnlineIdHandleTags::FVerifiedAuthTicket>;
+using FVerifiedAuthSessionId = TOnlineId<OnlineIdHandleTags::FVerifiedAuthSession>;
 
 COREONLINE_API FString ToLogString(const FAccountId& Id);
 COREONLINE_API FString ToLogString(const FLobbyId& Id);
 COREONLINE_API FString ToLogString(const FOnlineSessionId& Id);
-COREONLINE_API FString ToLogString(const FOnlineSessionInviteIdHandle& Id);
-COREONLINE_API FString ToLogString(const FOnlineVerifiedAuthTicketIdHandle& Id);
-COREONLINE_API FString ToLogString(const FOnlineVerifiedAuthSessionIdHandle& Id);
+COREONLINE_API FString ToLogString(const FSessionInviteId& Id);
+COREONLINE_API FString ToLogString(const FVerifiedAuthTicketId& Id);
+COREONLINE_API FString ToLogString(const FVerifiedAuthSessionId& Id);
 
 template<typename IdType>
-inline uint32 GetTypeHash(const TOnlineIdHandle<IdType>& Handle)
+inline uint32 GetTypeHash(const TOnlineId<IdType>& OnlineId)
 {
 	using ::GetTypeHash;
-	return HashCombine(GetTypeHash(Handle.GetOnlineServicesType()), GetTypeHash(Handle.GetHandle()));
+	return HashCombine(GetTypeHash(OnlineId.GetOnlineServicesType()), GetTypeHash(OnlineId.GetHandle()));
 }
 
 template<typename IdType>
@@ -333,9 +333,9 @@ class IOnlineIdRegistry
 public:
 	virtual ~IOnlineIdRegistry() = default;
 
-	virtual FString ToLogString(const TOnlineIdHandle<IdType>& Handle) const = 0;
-	virtual TArray<uint8> ToReplicationData(const TOnlineIdHandle<IdType>& Handle) const = 0;
-	virtual TOnlineIdHandle<IdType> FromReplicationData(const TArray<uint8>& ReplicationData) = 0;
+	virtual FString ToLogString(const TOnlineId<IdType>& OnlineId) const = 0;
+	virtual TArray<uint8> ToReplicationData(const TOnlineId<IdType>& OnlineId) const = 0;
+	virtual TOnlineId<IdType> FromReplicationData(const TArray<uint8>& ReplicationData) = 0;
 };
 
 using IOnlineAccountIdRegistry = IOnlineIdRegistry<OnlineIdHandleTags::FAccount>;
@@ -420,9 +420,9 @@ public:
 	 */
 	COREONLINE_API void UnregisterSessionInviteIdRegistry(EOnlineServices OnlineServices, int32 Priority = 0);
 
-	COREONLINE_API FString ToLogString(const FOnlineSessionInviteIdHandle& Handle) const;
-	COREONLINE_API TArray<uint8> ToReplicationData(const FOnlineSessionInviteIdHandle& Handle) const;
-	COREONLINE_API FOnlineSessionInviteIdHandle ToSessionInviteId(EOnlineServices Services, const TArray<uint8>& RepData) const;
+	COREONLINE_API FString ToLogString(const FSessionInviteId& SessionInviteId) const;
+	COREONLINE_API TArray<uint8> ToReplicationData(const FSessionInviteId& SessionInviteId) const;
+	COREONLINE_API FSessionInviteId ToSessionInviteId(EOnlineServices Services, const TArray<uint8>& RepData) const;
 
 	COREONLINE_API IOnlineSessionInviteIdRegistry* GetSessionInviteIdRegistry(EOnlineServices OnlineServices) const;
 
