Added a reference-get for NetId wrapper that can be used once we've verified the types and updated relevant call sites. This otherwise incurs a copy that takes about 40% of the cycles of operator==.

diff --git a/Engine/Source/Runtime/CoreOnline/Public/Online/CoreOnline.h b/Engine/Source/Runtime/CoreOnline/Public/Online/CoreOnline.h
--- a/Engine/Source/Runtime/CoreOnline/Public/Online/CoreOnline.h
+++ b/Engine/Source/Runtime/CoreOnline/Public/Online/CoreOnline.h
@@ -518,6 +518,12 @@ struct FUniqueNetIdWrapper
 		return Result;
 	}
 
+	// Getter to be used only when the variant index has already been confirmed
+	const FUniqueNetIdPtr& GetV1Unsafe() const
+	{
+		return Variant.Get<FUniqueNetIdPtr>();
+	}
+
 	bool IsV2() const
 	{
 		return Variant.IsType<UE::Online::FAccountId>();
@@ -533,12 +539,18 @@ struct FUniqueNetIdWrapper
 		return Result;
 	}
 
+	// Getter to be used only when the variant index has already been confirmed
+	const UE::Online::FAccountId& GetV2Unsafe() const
+	{
+		return Variant.Get<UE::Online::FAccountId>();
+	}
+
 	FName GetType() const
 	{
 		FName Result = NAME_None;
 		if (IsValid() && ensure(IsV1()))
 		{
-			Result = GetV1()->GetType();
+			Result = GetV1Unsafe()->GetType();
 		}
 		return Result;
 	}
@@ -549,7 +561,7 @@ struct FUniqueNetIdWrapper
 		FString Result;
 		if(IsValid() && ensure(IsV1()))
 		{
-			Result = GetV1()->ToString();
+			Result = GetV1Unsafe()->ToString();
 		}
 		return Result;
 	}
@@ -562,12 +574,12 @@ struct FUniqueNetIdWrapper
 	{
 		if (IsV1())
 		{
-			const FUniqueNetIdPtr& Ptr = GetV1();
+			const FUniqueNetIdPtr& Ptr = GetV1Unsafe();
 			return Ptr.IsValid() && Ptr->IsValid();
 		}
 		else
 		{
-			const UE::Online::FAccountId& AccountId = GetV2();
+			const UE::Online::FAccountId& AccountId = GetV2Unsafe();
 			return AccountId.IsValid();
 		}
 	}
@@ -618,11 +630,11 @@ struct FUniqueNetIdWrapper
 		{
 			if (Value.IsV1())
 			{
-				return GetTypeHash(*Value.GetV1());
+				return GetTypeHash(*Value.GetV1Unsafe());
 			}
 			else
 			{
-				return GetTypeHash(Value.GetV2());
+				return GetTypeHash(Value.GetV2Unsafe());
 			}
 		}
 		return INDEX_NONE;
@@ -656,11 +668,11 @@ struct FUniqueNetIdWrapper
 		if (Lhs.IsV1())
 		{
 			// Pointers can point to equivalent objects
-			return *Lhs.GetV1() == *Rhs.GetV1();
+			return *Lhs.GetV1Unsafe() == *Rhs.GetV1Unsafe();
 		}
 		else
 		{
-			return Lhs.GetV2() == Rhs.GetV2();
+			return Lhs.GetV2Unsafe() == Rhs.GetV2Unsafe();
 		}
 	}
 
@@ -687,7 +699,7 @@ struct FUniqueNetIdWrapper
 			// Different variant
 			return false;
 		}
-		return *Lhs.GetV1() == Rhs;
+		return *Lhs.GetV1Unsafe() == Rhs;
 	}
 
 	friend bool operator!=(const FUniqueNetIdWrapper& Lhs, const FUniqueNetId& Rhs)
