- Add FUniqueNetId::IsMutable, default to returning false.
- Return true from FUniqueNetIdEOS::IsMutable
- Make FUniqueNetIdRepl always rebuild replication data when IsMutable is false

diff --git a/Engine/Plugins/Online/OnlineSubsystemEOS/Source/OnlineSubsystemEOS/Private/OnlineSubsystemEOSTypes.cpp b/Engine/Plugins/Online/OnlineSubsystemEOS/Source/OnlineSubsystemEOS/Private/OnlineSubsystemEOSTypes.cpp
--- a/Engine/Plugins/Online/OnlineSubsystemEOS/Source/OnlineSubsystemEOS/Private/OnlineSubsystemEOSTypes.cpp
+++ b/Engine/Plugins/Online/OnlineSubsystemEOS/Source/OnlineSubsystemEOS/Private/OnlineSubsystemEOSTypes.cpp
@@ -49,6 +49,11 @@ uint32 FUniqueNetIdEOS::GetTypeHash() const
 	return ::GetTypeHash(static_cast<const void*>(this));
 }
 
+bool FUniqueNetIdEOS::IsMutable() const
+{
+	return true;
+}
+
 FString FUniqueNetIdEOS::ToString() const
 {
 	return LexToString(EpicAccountId) + EOS_ID_SEPARATOR + LexToString(ProductUserId);
diff --git a/Engine/Plugins/Online/OnlineSubsystemEOS/Source/OnlineSubsystemEOS/Private/OnlineSubsystemEOSTypes.h b/Engine/Plugins/Online/OnlineSubsystemEOS/Source/OnlineSubsystemEOS/Private/OnlineSubsystemEOSTypes.h
--- a/Engine/Plugins/Online/OnlineSubsystemEOS/Source/OnlineSubsystemEOS/Private/OnlineSubsystemEOSTypes.h
+++ b/Engine/Plugins/Online/OnlineSubsystemEOS/Source/OnlineSubsystemEOS/Private/OnlineSubsystemEOSTypes.h
@@ -43,6 +43,7 @@ public:
 	virtual int32 GetSize() const override;
 	virtual bool IsValid() const override;
 	virtual uint32 GetTypeHash() const override;
+	virtual bool IsMutable() const override;
 	virtual FString ToString() const override;
 	virtual FString ToDebugString() const override;
 
diff --git a/Engine/Source/Runtime/CoreOnline/Public/Online/CoreOnline.h b/Engine/Source/Runtime/CoreOnline/Public/Online/CoreOnline.h
--- a/Engine/Source/Runtime/CoreOnline/Public/Online/CoreOnline.h
+++ b/Engine/Source/Runtime/CoreOnline/Public/Online/CoreOnline.h
@@ -185,11 +185,32 @@ public:
 	 */
 	virtual FString ToDebugString() const = 0;
 
+	/**
+	 * Get a hash of the opaque data
+	 *
+	 * @return the hash
+	 */
 	virtual uint32 GetTypeHash() const
 	{
 		return CityHash32(reinterpret_cast<const char*>(GetBytes()), GetSize());
 	}
 
+	/**
+	 * Whether the opaque data may mutate.
+	 * 
+	 * If true, it is not safe to assume that GetBytes and GetSize will always return the same values.
+	 * Note however the following contracts which always apply:
+	 *  - GetTypeHash must be stable, and return the same value pre and post mutation
+	 *  - Compare of 2 net id's must return the same value pre and post mutation.
+	 *  - A net id cannot mutate what it logically represents, i.e. a net id representing player A cannot mutate to represent player B.
+	 * 
+	 * @return whether the opaque data may mutate
+	 */
+	virtual bool IsMutable() const
+	{
+		return false;
+	}
+
 	/**
 	 * @return hex encoded string representation of unique id
 	 */
diff --git a/Engine/Source/Runtime/Engine/Classes/GameFramework/OnlineReplStructs.h b/Engine/Source/Runtime/Engine/Classes/GameFramework/OnlineReplStructs.h
--- a/Engine/Source/Runtime/Engine/Classes/GameFramework/OnlineReplStructs.h
+++ b/Engine/Source/Runtime/Engine/Classes/GameFramework/OnlineReplStructs.h
@@ -124,6 +124,9 @@ protected:
 	/** Network serialized data cache */
 	UPROPERTY(Transient)
 	TArray<uint8> ReplicationBytes;
+
+	/** Do we need to rebuild the ReplicationData */
+	bool IsReplicationDataDirty() const;
 };
 
 /** Specify type trait support for various low level UPROPERTY overrides */
diff --git a/Engine/Source/Runtime/Engine/Private/OnlineReplStructs.cpp b/Engine/Source/Runtime/Engine/Private/OnlineReplStructs.cpp
--- a/Engine/Source/Runtime/Engine/Private/OnlineReplStructs.cpp
+++ b/Engine/Source/Runtime/Engine/Private/OnlineReplStructs.cpp
@@ -244,6 +244,21 @@ void FUniqueNetIdRepl::MakeReplicationDataV2()
 	Writer << ReplicationData;
 }
 
+bool FUniqueNetIdRepl::IsReplicationDataDirty() const
+{
+	if (ReplicationBytes.IsEmpty())
+	{
+		// If ReplicationBytes is empty, it is dirty by definition. Even invalid net ids have encoding flags.
+		return true;
+	}
+	if (IsValid() && IsV1() && GetUniqueNetId()->IsMutable())
+	{
+		// Mutable net id's require us to reserialize fresh every time.
+		return true;
+	}
+	return false;
+}
+
 void FUniqueNetIdRepl::UniqueIdFromString(FName Type, const FString& Contents)
 {
 	// Don't need to distinguish OSS interfaces here with world because we just want the create function below
@@ -257,7 +272,7 @@ bool FUniqueNetIdRepl::NetSerialize(FArchive& Ar, UPackageMap* Map, bool& bOutSu
 
 	if (Ar.IsSaving())
 	{
-		if (ReplicationBytes.Num() == 0)
+		if (IsReplicationDataDirty())
 		{
 			MakeReplicationData();
 		}
