Remove Stadia references

diff --git a/Engine/Source/Runtime/CoreOnline/Private/Online/CoreOnline.cpp b/Engine/Source/Runtime/CoreOnline/Private/Online/CoreOnline.cpp
--- a/Engine/Source/Runtime/CoreOnline/Private/Online/CoreOnline.cpp
+++ b/Engine/Source/Runtime/CoreOnline/Private/Online/CoreOnline.cpp
@@ -16,7 +16,7 @@ const TCHAR* LexToString(EOnlineServices Value)
 	case EOnlineServices::Xbox:				return TEXT("Xbox");
 	case EOnlineServices::PSN:				return TEXT("PSN");
 	case EOnlineServices::Nintendo:			return TEXT("Nintendo");
-	case EOnlineServices::Stadia:			return TEXT("Stadia");
+	case EOnlineServices::Reserved_5:		return TEXT("Reserved_5");
 	case EOnlineServices::Steam:			return TEXT("Steam");
 	case EOnlineServices::Google:			return TEXT("Google");
 	case EOnlineServices::GooglePlay:		return TEXT("GooglePlay");
@@ -72,9 +72,9 @@ void LexFromString(EOnlineServices& OutValue, const TCHAR* InStr)
 	{
 		OutValue = EOnlineServices::Nintendo;
 	}
-	else if (FCString::Stricmp(InStr, TEXT("Stadia")) == 0)
+	else if (FCString::Stricmp(InStr, TEXT("Reserved_5")) == 0)
 	{
-		OutValue = EOnlineServices::Stadia;
+		OutValue = EOnlineServices::Reserved_5;
 	}
 	else if (FCString::Stricmp(InStr, TEXT("Steam")) == 0)
 	{
diff --git a/Engine/Source/Runtime/CoreOnline/Public/Online/CoreOnline.h b/Engine/Source/Runtime/CoreOnline/Public/Online/CoreOnline.h
--- a/Engine/Source/Runtime/CoreOnline/Public/Online/CoreOnline.h
+++ b/Engine/Source/Runtime/CoreOnline/Public/Online/CoreOnline.h
@@ -235,8 +235,8 @@ enum class EOnlineServices : uint8
 	PSN,
 	// Nintendo
 	Nintendo,
-	// Stadia,
-	Stadia,
+	// Unused,
+	Reserved_5,
 	// Steam
 	Steam,
 	// Google
