- Add ToString for V2 net ids

diff --git a/Engine/Plugins/Online/OnlineServices/Source/OnlineServicesCommon/Public/Online/SessionsCommon.h b/Engine/Plugins/Online/OnlineServices/Source/OnlineServicesCommon/Public/Online/SessionsCommon.h
--- a/Engine/Plugins/Online/OnlineServices/Source/OnlineServicesCommon/Public/Online/SessionsCommon.h
+++ b/Engine/Plugins/Online/OnlineServices/Source/OnlineServicesCommon/Public/Online/SessionsCommon.h
@@ -16,18 +16,23 @@ class FOnlineSessionIdStringRegistry : public IOnlineSessionIdRegistry
 {
 public:
 	// Begin IOnlineSessionIdRegistry
-	virtual inline FString ToLogString(const FOnlineSessionId& SessionId) const override
+	virtual inline FString ToString(const FOnlineSessionId& SessionId) const override
 	{
 		FString IdValue = BasicRegistry.FindIdValue(SessionId);
 
-		if (IdValue.Len() == 0)
+		if (IdValue.IsEmpty())
 		{
-			IdValue = FString(TEXT("[InvalidSessionID]"));
+			IdValue = FString(TEXT("Invalid"));
 		}
 
 		return IdValue;
 	};
 
+	virtual inline FString ToLogString(const FOnlineSessionId& SessionId) const override
+	{
+		return ToString(SessionId);
+	};
+
 	virtual inline TArray<uint8> ToReplicationData(const FOnlineSessionId& SessionId) const override
 	{
 		const FString IdValue = BasicRegistry.FindIdValue(SessionId);
@@ -76,16 +81,21 @@ class FOnlineSessionInviteIdStringRegistry : public IOnlineSessionInviteIdRegist
 {
 public:
 	// Begin IOnlineSessionIdRegistry
-	virtual inline FString ToLogString(const FSessionInviteId& SessionInviteId) const override
+	virtual inline FString ToString(const FSessionInviteId& SessionInviteId) const override
 	{
 		FString IdValue = BasicRegistry.FindIdValue(SessionInviteId);
 
-		if (IdValue.Len() == 0)
+		if (IdValue.IsEmpty())
 		{
-			IdValue = FString(TEXT("[InvalidSessionID]"));
+			IdValue = TEXT("Invalid");
 		}
 
 		return IdValue;
+	}
+
+	virtual inline FString ToLogString(const FSessionInviteId& SessionInviteId) const override
+	{
+		return ToString(SessionInviteId);
 	};
 
 	virtual inline TArray<uint8> ToReplicationData(const FSessionInviteId& SessionInviteId) const override
diff --git a/Engine/Plugins/Online/OnlineServicesEOS/Source/Private/Online/OnlineIdEOS.cpp b/Engine/Plugins/Online/OnlineServicesEOS/Source/Private/Online/OnlineIdEOS.cpp
--- a/Engine/Plugins/Online/OnlineServicesEOS/Source/Private/Online/OnlineIdEOS.cpp
+++ b/Engine/Plugins/Online/OnlineServicesEOS/Source/Private/Online/OnlineIdEOS.cpp
@@ -8,6 +8,22 @@
 
 namespace UE::Online {
 
+FString FOnlineAccountIdRegistryEOS::ToString(const FAccountId& AccountId) const
+{
+	FString Result;
+	if (ValidateOnlineId(AccountId))
+	{
+		const FOnlineAccountIdDataEOS& IdData = GetAccountIdData(AccountId);
+		Result = FString::Printf(TEXT("%s|%s"), *LexToString(IdData.EpicAccountId), *LexToString(IdData.ProductUserId));
+	}
+	else
+	{
+		check(!AccountId.IsValid()); // Check we haven't been passed a valid handle for a different EOnlineServices.
+		Result = TEXT("Invalid");
+	}
+	return Result;
+}
+
 FString FOnlineAccountIdRegistryEOS::ToLogString(const FAccountId& AccountId) const
 {
 	FString Result;
diff --git a/Engine/Plugins/Online/OnlineServicesEOS/Source/Private/Online/OnlineIdEOS.h b/Engine/Plugins/Online/OnlineServicesEOS/Source/Private/Online/OnlineIdEOS.h
--- a/Engine/Plugins/Online/OnlineServicesEOS/Source/Private/Online/OnlineIdEOS.h
+++ b/Engine/Plugins/Online/OnlineServicesEOS/Source/Private/Online/OnlineIdEOS.h
@@ -27,6 +27,7 @@ public:
 	virtual ~FOnlineAccountIdRegistryEOS() = default;
 
 	// Begin IOnlineAccountIdRegistry
+	virtual FString ToString(const FAccountId& AccountId) const override;
 	virtual FString ToLogString(const FAccountId& AccountId) const override;
 	virtual TArray<uint8> ToReplicationData(const FAccountId& AccountId) const override;
 	virtual FAccountId FromReplicationData(const TArray<uint8>& ReplicationString) override;
diff --git a/Engine/Plugins/Online/OnlineServicesEOSGS/Source/Private/Online/OnlineIdEOSGS.cpp b/Engine/Plugins/Online/OnlineServicesEOSGS/Source/Private/Online/OnlineIdEOSGS.cpp
--- a/Engine/Plugins/Online/OnlineServicesEOSGS/Source/Private/Online/OnlineIdEOSGS.cpp
+++ b/Engine/Plugins/Online/OnlineServicesEOSGS/Source/Private/Online/OnlineIdEOSGS.cpp
@@ -27,14 +27,25 @@ EOS_ProductUserId FOnlineAccountIdRegistryEOSGS::GetProductUserId(const FAccount
 	return Registry.FindIdValue(AccountId);
 }
 
-FString FOnlineAccountIdRegistryEOSGS::ToLogString(const FAccountId& AccountId) const
+FString FOnlineAccountIdRegistryEOSGS::ToString(const FAccountId& AccountId) const
 {
+	FString Result;
 	if (Registry.ValidateOnlineId(AccountId))
 	{
 		EOS_ProductUserId ProductUserId = Registry.FindIdValue(AccountId);
-		return LexToString(ProductUserId);
+		Result = LexToString(ProductUserId);
+	}
+	else
+	{
+		check(!AccountId.IsValid()); // Check we haven't been passed a valid handle for a different EOnlineServices.
+		Result = TEXT("Invalid");
 	}
-	return FString();
+	return Result;
+}
+
+FString FOnlineAccountIdRegistryEOSGS::ToLogString(const FAccountId& AccountId) const
+{
+	return ToString(AccountId);
 }
 
 TArray<uint8> FOnlineAccountIdRegistryEOSGS::ToReplicationData(const FAccountId& AccountId) const
diff --git a/Engine/Plugins/Online/OnlineServicesEOSGS/Source/Public/Online/OnlineIdEOSGS.h b/Engine/Plugins/Online/OnlineServicesEOSGS/Source/Public/Online/OnlineIdEOSGS.h
--- a/Engine/Plugins/Online/OnlineServicesEOSGS/Source/Public/Online/OnlineIdEOSGS.h
+++ b/Engine/Plugins/Online/OnlineServicesEOSGS/Source/Public/Online/OnlineIdEOSGS.h
@@ -35,6 +35,7 @@ public:
 	// End IOnlineAccountIdRegistryEOSGS
 
 	// Begin IOnlineAccountIdRegistry
+	virtual FString ToString(const FAccountId& AccountId) const override;
 	virtual FString ToLogString(const FAccountId& AccountId) const override;
 	virtual TArray<uint8> ToReplicationData(const FAccountId& AccountId) const override;
 	virtual FAccountId FromReplicationData(const TArray<uint8>& ReplicationData) override;
diff --git a/Engine/Plugins/Online/OnlineServicesNull/Source/Private/Online/AuthNull.cpp b/Engine/Plugins/Online/OnlineServicesNull/Source/Private/Online/AuthNull.cpp
--- a/Engine/Plugins/Online/OnlineServicesNull/Source/Private/Online/AuthNull.cpp
+++ b/Engine/Plugins/Online/OnlineServicesNull/Source/Private/Online/AuthNull.cpp
@@ -227,7 +227,7 @@ FAccountId FOnlineAccountIdRegistryNull::FindOrAddAccountId(const FString& Accou
 	}
 }
 
-FString FOnlineAccountIdRegistryNull::ToLogString(const FAccountId& AccountId) const
+FString FOnlineAccountIdRegistryNull::ToString(const FAccountId& AccountId) const
 {
 	const FReadScopeLock ReadLock(Lock);
 	if(const FOnlineAccountIdString* Id = FindNoLock(AccountId))
@@ -235,9 +235,13 @@ FString FOnlineAccountIdRegistryNull::ToLogString(const FAccountId& AccountId) c
 		return Id->Data;
 	}
 
-	return FString(TEXT("[InvalidNetID]"));
+	return FString(TEXT("Invalid"));
 }
 
+FString FOnlineAccountIdRegistryNull::ToLogString(const FAccountId& AccountId) const
+{
+	return ToString(AccountId);
+}
 
 TArray<uint8> FOnlineAccountIdRegistryNull::ToReplicationData(const FAccountId& AccountId) const
 {
diff --git a/Engine/Plugins/Online/OnlineServicesNull/Source/Private/Online/LobbyRegistryNull.cpp b/Engine/Plugins/Online/OnlineServicesNull/Source/Private/Online/LobbyRegistryNull.cpp
--- a/Engine/Plugins/Online/OnlineServicesNull/Source/Private/Online/LobbyRegistryNull.cpp
+++ b/Engine/Plugins/Online/OnlineServicesNull/Source/Private/Online/LobbyRegistryNull.cpp
@@ -45,16 +45,20 @@ const FString* FOnlineLobbyIdRegistryNull::GetInternal(const FLobbyId& LobbyId)
 	return nullptr;
 }
 
-FString FOnlineLobbyIdRegistryNull::ToLogString(const FLobbyId& LobbyId) const
+FString FOnlineLobbyIdRegistryNull::ToString(const FLobbyId& LobbyId) const
 {
 	if (const FString* Id = GetInternal(LobbyId))
 	{
 		return *Id;
 	}
 
-	return FString(TEXT("[InvalidLobbyID]"));
+	return FString(TEXT("Invalid"));
 }
 
+FString FOnlineLobbyIdRegistryNull::ToLogString(const FLobbyId& LobbyId) const
+{
+	return ToString(LobbyId);
+}
 
 TArray<uint8> FOnlineLobbyIdRegistryNull::ToReplicationData(const FLobbyId& LobbyId) const
 {
diff --git a/Engine/Plugins/Online/OnlineServicesNull/Source/Public/Online/AuthNull.h b/Engine/Plugins/Online/OnlineServicesNull/Source/Public/Online/AuthNull.h
--- a/Engine/Plugins/Online/OnlineServicesNull/Source/Public/Online/AuthNull.h
+++ b/Engine/Plugins/Online/OnlineServicesNull/Source/Public/Online/AuthNull.h
@@ -27,6 +27,7 @@ public:
 	FAccountId FindOrAddAccountId(const FString& AccountId);
 
 	// Begin IOnlineAccountIdRegistry
+	virtual FString ToString(const FAccountId& AccountId) const override;
 	virtual FString ToLogString(const FAccountId& AccountId) const override;
 	virtual TArray<uint8> ToReplicationData(const FAccountId& AccountId) const override;
 	virtual FAccountId FromReplicationData(const TArray<uint8>& ReplicationString) override;
diff --git a/Engine/Plugins/Online/OnlineServicesNull/Source/Public/Online/LobbyRegistryNull.h b/Engine/Plugins/Online/OnlineServicesNull/Source/Public/Online/LobbyRegistryNull.h
--- a/Engine/Plugins/Online/OnlineServicesNull/Source/Public/Online/LobbyRegistryNull.h
+++ b/Engine/Plugins/Online/OnlineServicesNull/Source/Public/Online/LobbyRegistryNull.h
@@ -17,6 +17,7 @@ public:
 	FLobbyId GetNext();
 
 	// Begin IOnlineAccountIdRegistry
+	virtual FString ToString(const FLobbyId& LobbyId) const override;
 	virtual FString ToLogString(const FLobbyId& LobbyId) const override;
 	virtual TArray<uint8> ToReplicationData(const FLobbyId& LobbyId) const override;
 	virtual FLobbyId FromReplicationData(const TArray<uint8>& ReplicationString) override;
@@ -25,10 +26,10 @@ public:
 	virtual ~FOnlineLobbyIdRegistryNull() = default;
 
 private:
+	// TODO use TOnlineBasicIdRegistry
 	const FString* GetInternal(const FLobbyId& LobbyId) const;
 	TArray<FString> Ids;
-	TMap<FString, FLobbyId> StringToId; 
-
+	TMap<FString, FLobbyId> StringToId;
 };
 
 } // namespace UE::Online
\ No newline at end of file
diff --git a/Engine/Plugins/Online/OnlineServicesOSSAdapter/Source/Public/Online/OnlineIdOSSAdapter.h b/Engine/Plugins/Online/OnlineServicesOSSAdapter/Source/Public/Online/OnlineIdOSSAdapter.h
--- a/Engine/Plugins/Online/OnlineServicesOSSAdapter/Source/Public/Online/OnlineIdOSSAdapter.h
+++ b/Engine/Plugins/Online/OnlineServicesOSSAdapter/Source/Public/Online/OnlineIdOSSAdapter.h
@@ -84,10 +84,16 @@ public:
 	}
 
 	// Begin IOnlineAccountIdRegistry
+	virtual FString ToString(const TOnlineId<IdType>& OnlineId) const override
+	{
+		FUniqueNetIdPtr IdValue = GetIdValue(OnlineId);
+		return IdValue ? IdValue->ToString() : TEXT("invalid_id");
+	}
+
 	virtual FString ToLogString(const TOnlineId<IdType>& OnlineId) const override
 	{
 		FUniqueNetIdPtr IdValue = GetIdValue(OnlineId);
-		return IdValue ? IdValue->ToDebugString() : FString(TEXT("invalid_id"));
+		return IdValue ? IdValue->ToDebugString() : TEXT("invalid_id");
 	}
 
 	virtual TArray<uint8> ToReplicationData(const TOnlineId<IdType>& OnlineId) const override
@@ -99,7 +105,6 @@ public:
 	{
 		return TOnlineId<IdType>();
 	}
-
 	// End IOnlineAccountIdRegistry
 
 private:
diff --git a/Engine/Source/Runtime/CoreOnline/Private/Online/CoreOnline.cpp b/Engine/Source/Runtime/CoreOnline/Private/Online/CoreOnline.cpp
--- a/Engine/Source/Runtime/CoreOnline/Private/Online/CoreOnline.cpp
+++ b/Engine/Source/Runtime/CoreOnline/Private/Online/CoreOnline.cpp
@@ -237,6 +237,20 @@ void FOnlineIdRegistryRegistry::UnregisterAccountIdRegistry(EOnlineServices Onli
 	}
 }
 
+FString FOnlineIdRegistryRegistry::ToString(const FAccountId& AccountId) const
+{
+	FString Result;
+	if (IOnlineAccountIdRegistry* Registry = GetAccountIdRegistry(AccountId.GetOnlineServicesType()))
+	{
+		Result = Registry->ToString(AccountId);
+	}
+	else
+	{
+		Result = ForeignAccountIdRegistry->ToString(AccountId);
+	}
+	return Result;
+}
+
 FString FOnlineIdRegistryRegistry::ToLogString(const FAccountId& AccountId) const
 {
 	FString Result;
@@ -405,36 +419,42 @@ IOnlineSessionInviteIdRegistry* FOnlineIdRegistryRegistry::GetSessionInviteIdReg
 	return Registry;
 }
 
-FString ToLogString(const FAccountId& Id)
-{
-	return FString::Printf(TEXT("%s:%d (%s)"), LexToString(Id.GetOnlineServicesType()), Id.GetHandle(), *FOnlineIdRegistryRegistry::Get().ToLogString(Id));
-}
-
-FString ToLogString(const FLobbyId& Id)
+template<typename IdType>
+FString ToStringImpl(const TOnlineId<IdType>& Id)
 {
-	// todo: use registry.
-	return FString::Printf(TEXT("%s:%d"), LexToString(Id.GetOnlineServicesType()), Id.GetHandle());
+	return FOnlineIdRegistryRegistry::Get().ToString(Id);
 }
 
-FString ToLogString(const FOnlineSessionId& Id)
-{
-	return FString::Printf(TEXT("%s:%d (%s)"), LexToString(Id.GetOnlineServicesType()), Id.GetHandle(), *FOnlineIdRegistryRegistry::Get().ToLogString(Id));
-}
+FString ToString(const FAccountId& Id) { return ToStringImpl(Id); }
+// TODO
+//FString ToString(const FLobbyId& Id) { return ToStringImpl(Id); }
+//FString ToString(const FOnlineSessionId& Id) { return ToStringImpl(Id); }
+//FString ToString(const FSessionInviteId& Id) { return ToStringImpl(Id); }
+//FString ToString(const FVerifiedAuthTicketId& Id) { return ToStringImpl(Id); }
+//FString ToString(const FVerifiedAuthSessionId& Id) { return ToStringImpl(Id); }
 
-FString ToLogString(const FSessionInviteId& Id)
+template<typename IdType>
+FString ToLogStringImpl(const TOnlineId<IdType>& Id)
 {
-	return FString::Printf(TEXT("%s:%d (%s)"), LexToString(Id.GetOnlineServicesType()), Id.GetHandle(), *FOnlineIdRegistryRegistry::Get().ToLogString(Id));
-}
-
-FString ToLogString(const FVerifiedAuthTicketId& Id)
-{
-	return FString::Printf(TEXT("%s:%d"), LexToString(Id.GetOnlineServicesType()), Id.GetHandle());
+	if constexpr (std::is_same_v<IdType, FAccountId>
+		|| std::is_same_v<IdType, FOnlineSessionId>
+		|| std::is_same_v<IdType, FSessionInviteId>)
+		// TODO Add FLobbyId
+	{
+		return FString::Printf(TEXT("%s:%d (%s)"), LexToString(Id.GetOnlineServicesType()), Id.GetHandle(), *FOnlineIdRegistryRegistry::Get().ToLogString(Id));
+	}
+	else
+	{
+		return FString::Printf(TEXT("%s:%d"), LexToString(Id.GetOnlineServicesType()), Id.GetHandle());
+	}
 }
 
-FString ToLogString(const FVerifiedAuthSessionId& Id)
-{
-	return FString::Printf(TEXT("%s:%d"), LexToString(Id.GetOnlineServicesType()), Id.GetHandle());
-}
+FString ToLogString(const FAccountId& Id) { return ToLogStringImpl(Id); }
+FString ToLogString(const FLobbyId& Id) { return ToLogStringImpl(Id); }
+FString ToLogString(const FOnlineSessionId& Id) { return ToLogStringImpl(Id); }
+FString ToLogString(const FSessionInviteId& Id) { return ToLogStringImpl(Id); }
+FString ToLogString(const FVerifiedAuthTicketId& Id) { return ToLogStringImpl(Id); }
+FString ToLogString(const FVerifiedAuthSessionId& Id) { return ToLogStringImpl(Id); }
 
 }	/* UE::Online */
 
diff --git a/Engine/Source/Runtime/CoreOnline/Private/Online/CoreOnlinePrivate.cpp b/Engine/Source/Runtime/CoreOnline/Private/Online/CoreOnlinePrivate.cpp
--- a/Engine/Source/Runtime/CoreOnline/Private/Online/CoreOnlinePrivate.cpp
+++ b/Engine/Source/Runtime/CoreOnline/Private/Online/CoreOnlinePrivate.cpp
@@ -4,6 +4,27 @@
 
 namespace UE::Online {
 
+FString FOnlineForeignAccountIdRegistry::ToString(const FAccountId& AccountId) const
+{
+	FString Result;
+	if (AccountId.IsValid())
+	{
+		const EOnlineServices HandleType = AccountId.GetOnlineServicesType();
+		const uint32 HandleIndex = AccountId.GetHandle() - 1;
+		if (const FRepData* RepDataForServices = OnlineServicesToRepData.Find(HandleType))
+		{
+			const TArray<TArray<uint8>>& RepDataArray = RepDataForServices->RepDataArray;
+			if (ensure(RepDataArray.IsValidIndex(HandleIndex)))
+			{
+				const TArray<uint8>& RepData = RepDataArray[HandleIndex];
+				// TODO really this requires traits to interpret the bytes as a string/int/whatever
+				Result = FString::FromHexBlob(RepData.GetData(), RepData.Num());
+			}
+		}
+	}
+	return Result;
+}
+
 FString FOnlineForeignAccountIdRegistry::ToLogString(const FAccountId& AccountId) const
 {
 	FString Result;
diff --git a/Engine/Source/Runtime/CoreOnline/Private/Online/CoreOnlinePrivate.h b/Engine/Source/Runtime/CoreOnline/Private/Online/CoreOnlinePrivate.h
--- a/Engine/Source/Runtime/CoreOnline/Private/Online/CoreOnlinePrivate.h
+++ b/Engine/Source/Runtime/CoreOnline/Private/Online/CoreOnlinePrivate.h
@@ -30,6 +30,7 @@ namespace UE::Online {
 class FOnlineForeignAccountIdRegistry
 {
 public:
+	FString ToString(const FAccountId& AccountId) const;
 	FString ToLogString(const FAccountId& AccountId) const;
 	TArray<uint8> ToReplicationData(const FAccountId& AccountId) const;
 	FAccountId FromReplicationData(EOnlineServices Services, const TArray<uint8>& RepData);
diff --git a/Engine/Source/Runtime/CoreOnline/Public/Online/CoreOnline.h b/Engine/Source/Runtime/CoreOnline/Public/Online/CoreOnline.h
--- a/Engine/Source/Runtime/CoreOnline/Public/Online/CoreOnline.h
+++ b/Engine/Source/Runtime/CoreOnline/Public/Online/CoreOnline.h
@@ -314,11 +314,20 @@ private:
 
 using FAccountId = TOnlineId<OnlineIdHandleTags::FAccount>;
 using FLobbyId = TOnlineId<OnlineIdHandleTags::FLobby>;
+// TODO rename to FSessionId
 using FOnlineSessionId = TOnlineId<OnlineIdHandleTags::FSession>;
 using FSessionInviteId = TOnlineId<OnlineIdHandleTags::FSessionInvite>;
 using FVerifiedAuthTicketId = TOnlineId<OnlineIdHandleTags::FVerifiedAuthTicket>;
 using FVerifiedAuthSessionId = TOnlineId<OnlineIdHandleTags::FVerifiedAuthSession>;
 
+COREONLINE_API FString ToString(const FAccountId& Id);
+// TODO
+// COREONLINE_API FString ToString(const FLobbyId& Id);
+// COREONLINE_API FString ToString(const FOnlineSessionId& Id);
+// COREONLINE_API FString ToString(const FSessionInviteId& Id);
+// COREONLINE_API FString ToString(const FVerifiedAuthTicketId& Id);
+// COREONLINE_API FString ToString(const FVerifiedAuthSessionId& Id);
+
 COREONLINE_API FString ToLogString(const FAccountId& Id);
 COREONLINE_API FString ToLogString(const FLobbyId& Id);
 COREONLINE_API FString ToLogString(const FOnlineSessionId& Id);
@@ -339,6 +348,7 @@ class IOnlineIdRegistry
 public:
 	virtual ~IOnlineIdRegistry() = default;
 
+	virtual FString ToString(const TOnlineId<IdType>& OnlineId) const = 0;
 	virtual FString ToLogString(const TOnlineId<IdType>& OnlineId) const = 0;
 	virtual TArray<uint8> ToReplicationData(const TOnlineId<IdType>& OnlineId) const = 0;
 	virtual TOnlineId<IdType> FromReplicationData(const TArray<uint8>& ReplicationData) = 0;
@@ -380,6 +390,8 @@ public:
 	 */
 	COREONLINE_API void UnregisterAccountIdRegistry(EOnlineServices OnlineServices, int32 Priority = 0);
 
+	// TODO Might be worth these being templates
+	COREONLINE_API FString ToString(const FAccountId& AccountId) const;
 	COREONLINE_API FString ToLogString(const FAccountId& AccountId) const;
 	COREONLINE_API TArray<uint8> ToReplicationData(const FAccountId& AccountId) const;
 	COREONLINE_API FAccountId ToAccountId(EOnlineServices Services, const TArray<uint8>& RepData) const;
