- Fix To[Log]StringImpl comparing the wrong types.

diff --git a/Engine/Source/Runtime/CoreOnline/Private/Online/CoreOnline.cpp b/Engine/Source/Runtime/CoreOnline/Private/Online/CoreOnline.cpp
--- a/Engine/Source/Runtime/CoreOnline/Private/Online/CoreOnline.cpp
+++ b/Engine/Source/Runtime/CoreOnline/Private/Online/CoreOnline.cpp
@@ -420,7 +420,7 @@ IOnlineSessionInviteIdRegistry* FOnlineIdRegistryRegistry::GetSessionInviteIdReg
 }
 
 template<typename IdType>
-FString ToStringImpl(const TOnlineId<IdType>& Id)
+FString ToStringImpl(const IdType& Id)
 {
 	return FOnlineIdRegistryRegistry::Get().ToString(Id);
 }
@@ -434,7 +434,7 @@ FString ToString(const FAccountId& Id) { return ToStringImpl(Id); }
 //FString ToString(const FVerifiedAuthSessionId& Id) { return ToStringImpl(Id); }
 
 template<typename IdType>
-FString ToLogStringImpl(const TOnlineId<IdType>& Id)
+FString ToLogStringImpl(const IdType& Id)
 {
 	if constexpr (std::is_same_v<IdType, FAccountId>
 		|| std::is_same_v<IdType, FOnlineSessionId>
