- Add ToString for v2 session ids

diff --git a/Engine/Source/Runtime/CoreOnline/Private/Online/CoreOnline.cpp b/Engine/Source/Runtime/CoreOnline/Private/Online/CoreOnline.cpp
--- a/Engine/Source/Runtime/CoreOnline/Private/Online/CoreOnline.cpp
+++ b/Engine/Source/Runtime/CoreOnline/Private/Online/CoreOnline.cpp
@@ -321,6 +321,16 @@ void FOnlineIdRegistryRegistry::UnregisterSessionIdRegistry(EOnlineServices Onli
 	}
 }
 
+FString FOnlineIdRegistryRegistry::ToString(const FOnlineSessionId& SessionId) const
+{
+	FString Result;
+	if (IOnlineSessionIdRegistry* Registry = GetSessionIdRegistry(SessionId.GetOnlineServicesType()))
+	{
+		Result = Registry->ToString(SessionId);
+	}
+	return Result;
+}
+
 FString FOnlineIdRegistryRegistry::ToLogString(const FOnlineSessionId& SessionId) const
 {
 	FString Result;
@@ -426,9 +436,9 @@ FString ToStringImpl(const IdType& Id)
 }
 
 FString ToString(const FAccountId& Id) { return ToStringImpl(Id); }
+FString ToString(const FOnlineSessionId& Id) { return ToStringImpl(Id); }
 // TODO
 //FString ToString(const FLobbyId& Id) { return ToStringImpl(Id); }
-//FString ToString(const FOnlineSessionId& Id) { return ToStringImpl(Id); }
 //FString ToString(const FSessionInviteId& Id) { return ToStringImpl(Id); }
 //FString ToString(const FVerifiedAuthTicketId& Id) { return ToStringImpl(Id); }
 //FString ToString(const FVerifiedAuthSessionId& Id) { return ToStringImpl(Id); }
diff --git a/Engine/Source/Runtime/CoreOnline/Public/Online/CoreOnline.h b/Engine/Source/Runtime/CoreOnline/Public/Online/CoreOnline.h
--- a/Engine/Source/Runtime/CoreOnline/Public/Online/CoreOnline.h
+++ b/Engine/Source/Runtime/CoreOnline/Public/Online/CoreOnline.h
@@ -321,9 +321,9 @@ using FVerifiedAuthTicketId = TOnlineId<OnlineIdHandleTags::FVerifiedAuthTicket>
 using FVerifiedAuthSessionId = TOnlineId<OnlineIdHandleTags::FVerifiedAuthSession>;
 
 COREONLINE_API FString ToString(const FAccountId& Id);
+COREONLINE_API FString ToString(const FOnlineSessionId& Id);
 // TODO
 // COREONLINE_API FString ToString(const FLobbyId& Id);
-// COREONLINE_API FString ToString(const FOnlineSessionId& Id);
 // COREONLINE_API FString ToString(const FSessionInviteId& Id);
 // COREONLINE_API FString ToString(const FVerifiedAuthTicketId& Id);
 // COREONLINE_API FString ToString(const FVerifiedAuthSessionId& Id);
@@ -415,6 +415,7 @@ public:
 	 */
 	COREONLINE_API void UnregisterSessionIdRegistry(EOnlineServices OnlineServices, int32 Priority = 0);
 
+	COREONLINE_API FString ToString(const FOnlineSessionId& SessionId) const;
 	COREONLINE_API FString ToLogString(const FOnlineSessionId& SessionId) const;
 	COREONLINE_API TArray<uint8> ToReplicationData(const FOnlineSessionId& SessionId) const;
 	COREONLINE_API FOnlineSessionId ToSessionId(EOnlineServices Services, const TArray<uint8>& RepData) const;
