- Move FUniqueNetIdWrapper::ToString to .cpp, and add handling of V2 types

diff --git a/Engine/Source/Runtime/CoreOnline/Private/Online/CoreOnline.cpp b/Engine/Source/Runtime/CoreOnline/Private/Online/CoreOnline.cpp
--- a/Engine/Source/Runtime/CoreOnline/Private/Online/CoreOnline.cpp
+++ b/Engine/Source/Runtime/CoreOnline/Private/Online/CoreOnline.cpp
@@ -468,6 +468,23 @@ FString ToLogString(const FVerifiedAuthSessionId& Id) { return ToLogStringImpl(I
 
 }	/* UE::Online */
 
+FString FUniqueNetIdWrapper::ToString() const
+{
+	FString Result;
+	if (IsValid())
+	{
+		if (IsV1())
+		{
+			Result = GetV1Unsafe()->ToString();
+		}
+		else if (IsV2())
+		{
+			Result = UE::Online::ToString(GetV2Unsafe());
+		}
+	}
+	return Result;
+}
+
 FString FUniqueNetIdWrapper::ToDebugString() const
 {
 	FString Result;
diff --git a/Engine/Source/Runtime/CoreOnline/Public/Online/CoreOnline.h b/Engine/Source/Runtime/CoreOnline/Public/Online/CoreOnline.h
--- a/Engine/Source/Runtime/CoreOnline/Public/Online/CoreOnline.h
+++ b/Engine/Source/Runtime/CoreOnline/Public/Online/CoreOnline.h
@@ -570,15 +570,7 @@ struct FUniqueNetIdWrapper
 	}
 
 	/** Convert this value to a string */
-	FString ToString() const
-	{
-		FString Result;
-		if(IsValid() && ensure(IsV1()))
-		{
-			Result = GetV1Unsafe()->ToString();
-		}
-		return Result;
-	}
+	COREONLINE_API FString ToString() const;
 
 	/** Convert this value to a string with additional information */
 	COREONLINE_API FString ToDebugString() const;
