Move friend inlined GetTypeHash implementation for FUniqueNetIdWrapper outside of FUniqueNetIdWrapper.
Fixes issue when attempting to hash a FUniqueNetIdWrapper from inside of a function also named GetTypeHash

diff --git a/Engine/Source/Runtime/CoreOnline/Public/Online/CoreOnline.h b/Engine/Source/Runtime/CoreOnline/Public/Online/CoreOnline.h
--- a/Engine/Source/Runtime/CoreOnline/Public/Online/CoreOnline.h
+++ b/Engine/Source/Runtime/CoreOnline/Public/Online/CoreOnline.h
@@ -627,25 +627,6 @@ struct FUniqueNetIdWrapper
 		return GetV1().Get();
 	}
 
-	/**
-	* Friend function for using FUniqueNetIdWrapper as a hashable key
-	*/
-	friend inline uint32 GetTypeHash(const FUniqueNetIdWrapper& Value)
-	{
-		if (Value.IsValid())
-		{
-			if (Value.IsV1())
-			{
-				return GetTypeHash(*Value.GetV1Unsafe());
-			}
-			else
-			{
-				return GetTypeHash(Value.GetV2Unsafe());
-			}
-		}
-		return INDEX_NONE;
-	}
-
 	static FUniqueNetIdWrapper Invalid()
 	{
 		static FUniqueNetIdWrapper InvalidId(nullptr);
@@ -750,6 +731,25 @@ protected:
 	FVariantType Variant;
 };
 
+/**
+ * Function for using FUniqueNetIdWrapper as a hashable key
+ */
+inline uint32 GetTypeHash(const FUniqueNetIdWrapper& Value)
+{
+	if (Value.IsValid())
+	{
+		if (Value.IsV1())
+		{
+			return GetTypeHash(*Value.GetV1Unsafe());
+		}
+		else
+		{
+			return GetTypeHash(Value.GetV2Unsafe());
+		}
+	}
+	return INDEX_NONE;
+}
+
 template <typename ValueType>
 struct TUniqueNetIdMapKeyFuncs : public TDefaultMapKeyFuncs<FUniqueNetIdRef, ValueType, false>
 {
