Need to autortfm-ifi CoreOnline

diff --git a/Engine/Source/Runtime/CoreOnline/CoreOnline.Build.cs b/Engine/Source/Runtime/CoreOnline/CoreOnline.Build.cs
--- a/Engine/Source/Runtime/CoreOnline/CoreOnline.Build.cs
+++ b/Engine/Source/Runtime/CoreOnline/CoreOnline.Build.cs
@@ -22,6 +22,8 @@ public class CoreOnline : ModuleRules
 
 		PrivateDefinitions.Add("COREONLINE_PACKAGE=1");
 		PublicDefinitions.Add("PLATFORM_MAX_LOCAL_PLAYERS=" + GetPlatformMaxLocalPlayers(Target));
+
+		bAllowAutoRTFMInstrumentation = true;
 	}
 
 	protected virtual int GetPlatformMaxLocalPlayers(ReadOnlyTargetRules Target)
