- Add EpicAccount and EpicGame entries to EOnlineServices

diff --git a/Engine/Source/Runtime/CoreOnline/Private/Online/CoreOnline.cpp b/Engine/Source/Runtime/CoreOnline/Private/Online/CoreOnline.cpp
--- a/Engine/Source/Runtime/CoreOnline/Private/Online/CoreOnline.cpp
+++ b/Engine/Source/Runtime/CoreOnline/Private/Online/CoreOnline.cpp
@@ -25,8 +25,8 @@ const TCHAR* LexToString(EOnlineServices Value)
 	case EOnlineServices::Samsung:			return TEXT("Samsung");
 	case EOnlineServices::Oculus:			return TEXT("Oculus");
 	case EOnlineServices::Tencent:			return TEXT("Tencent");
-	case EOnlineServices::Reserved_14:		return TEXT("Reserved_14");
-	case EOnlineServices::Reserved_15:		return TEXT("Reserved_15");
+	case EOnlineServices::EpicAccount:		return TEXT("EpicAccount");
+	case EOnlineServices::EpicGame:			return TEXT("EpicGame");
 	case EOnlineServices::Reserved_16:		return TEXT("Reserved_16");
 	case EOnlineServices::Reserved_17:		return TEXT("Reserved_17");
 	case EOnlineServices::Reserved_18:		return TEXT("Reserved_18");
@@ -108,13 +108,13 @@ void LexFromString(EOnlineServices& OutValue, const TCHAR* InStr)
 	{
 		OutValue = EOnlineServices::Tencent;
 	}
-	else if (FCString::Stricmp(InStr, TEXT("Reserved_14")) == 0)
+	else if (FCString::Stricmp(InStr, TEXT("EpicAccount")) == 0)
 	{
-		OutValue = EOnlineServices::Reserved_14;
+		OutValue = EOnlineServices::EpicAccount;
 	}
-	else if (FCString::Stricmp(InStr, TEXT("Reserved_15")) == 0)
+	else if (FCString::Stricmp(InStr, TEXT("EpicGame")) == 0)
 	{
-		OutValue = EOnlineServices::Reserved_15;
+		OutValue = EOnlineServices::EpicGame;
 	}
 	else if (FCString::Stricmp(InStr, TEXT("Reserved_16")) == 0)
 	{
diff --git a/Engine/Source/Runtime/CoreOnline/Public/Online/CoreOnline.h b/Engine/Source/Runtime/CoreOnline/Public/Online/CoreOnline.h
--- a/Engine/Source/Runtime/CoreOnline/Public/Online/CoreOnline.h
+++ b/Engine/Source/Runtime/CoreOnline/Public/Online/CoreOnline.h
@@ -220,7 +220,7 @@ enum class EOnlineServices : uint8
 {
 	// Null, Providing minimal functionality when no backend services are required
 	Null,
-	// Epic Online Services
+	// Epic Online Services, merges Epic Account Services and Epic Game Services. Note this will be deprecated in favor of separating out the services
 	Epic,
 	// Xbox services
 	Xbox,
@@ -246,9 +246,11 @@ enum class EOnlineServices : uint8
 	Oculus,
 	// Tencent
 	Tencent,
+	// Epic Account Services (friends, presence, etc)
+	EpicAccount,
+	// Epic Game Services (lobbies, achievements, etc)
+	EpicGame,
 	// Reserved for future use/platform extensions
-	Reserved_14,
-	Reserved_15,
 	Reserved_16,
 	Reserved_17,
 	Reserved_18,
