Added support to register alias for EOnlineServices::GameDefined_x types, to be able to get it through LexToString/LexFromString.

diff --git a/Engine/Source/Runtime/CoreOnline/Private/Online/CoreOnline.cpp b/Engine/Source/Runtime/CoreOnline/Private/Online/CoreOnline.cpp
--- a/Engine/Source/Runtime/CoreOnline/Private/Online/CoreOnline.cpp
+++ b/Engine/Source/Runtime/CoreOnline/Private/Online/CoreOnline.cpp
@@ -5,10 +5,34 @@
 #include "Misc/LazySingleton.h"
 #include "Online/CoreOnlinePrivate.h"
 
+#define UE_ONLINE_ONLINESERVICES_MAX_LENGTH 16
+
 namespace UE::Online {
 
+static const uint32 OnlineServicesAliasArraySize = 4;
+TCHAR OnlineServicesAliasArray[OnlineServicesAliasArraySize][UE_ONLINE_ONLINESERVICES_MAX_LENGTH] = {};
+
+void RegisterAliasForGameDefined(EOnlineServices Value, const TCHAR* InStr)
+{
+	uint32 Index = (uint8)Value - (uint8)EOnlineServices::GameDefined_0;
+	check(Index >= 0 && Index < OnlineServicesAliasArraySize);
+	check(FCString::Strlen(InStr) < UE_ONLINE_ONLINESERVICES_MAX_LENGTH);
+	FCString::Strcpy(&OnlineServicesAliasArray[Index][0], UE_ONLINE_ONLINESERVICES_MAX_LENGTH, InStr);
+}
+
 const TCHAR* LexToString(EOnlineServices Value)
 {
+	if (Value >= EOnlineServices::GameDefined_0 && Value <= EOnlineServices((uint8)EOnlineServices::GameDefined_0 + OnlineServicesAliasArraySize))
+	{
+		uint32 Index = (uint8)Value - (uint8)EOnlineServices::GameDefined_0;
+		check(Index >= 0 && Index < OnlineServicesAliasArraySize);
+		const TCHAR* Alias = &OnlineServicesAliasArray[Index][0];
+		if (Alias[0] != TEXT('\0'))
+		{
+			return Alias;
+		}
+	}
+
 	switch (Value)
 	{
 	case EOnlineServices::Null:				return TEXT("Null");
@@ -194,6 +218,16 @@ void LexFromString(EOnlineServices& OutValue, const TCHAR* InStr)
 	}
 	else
 	{
+		for (uint32 Index = 0; Index < OnlineServicesAliasArraySize; ++Index)
+		{
+			const TCHAR* Alias = &OnlineServicesAliasArray[Index][0];
+			if (FCString::Stricmp(InStr, Alias) == 0)
+			{
+				OutValue = EOnlineServices((uint32)EOnlineServices::GameDefined_0 + Index);
+				return;
+			}
+		}
+
 		// Empty string is ok and gets mapped to "None". Anything else implies we are missing a case.
 		check(FCString::Strlen(InStr) == 0);
 		OutValue = EOnlineServices::None;
diff --git a/Engine/Source/Runtime/CoreOnline/Public/Online/CoreOnline.h b/Engine/Source/Runtime/CoreOnline/Public/Online/CoreOnline.h
--- a/Engine/Source/Runtime/CoreOnline/Public/Online/CoreOnline.h
+++ b/Engine/Source/Runtime/CoreOnline/Public/Online/CoreOnline.h
@@ -264,6 +264,18 @@ enum class EOnlineServices : uint8
 	Reserved_26,
 	Reserved_27,
 	// For game specific Online Services
+	// To make use of GameDefined_x, ideally it should do all these steps in module startup:
+	//	1. Read the EOnlineServices value from config, to avoid conflict with any other plugin while keeping it data-driven;
+	//			EOnlineServices XxxServices = EOnlineServices::GameDefined_0;
+	//			FString ConfigStr;
+	//			if (GConfig->GetString(TEXT("OnlineServicesXxx"), TEXT("ServicesEnumEntry"), ConfigStr, GEngineIni))
+	//			{
+	//				LexFromString(XxxServices, *ConfigStr);
+	//			}
+	//	2. Register the factory with that value;
+	//			FOnlineServicesRegistry::Get().RegisterServicesFactory(XxxServices, ...);
+	//	3. Optionally register the alias based on the value read from config  
+	//			RegisterAliasForGameDefined(XxxServices, TEXT("Xxx"));
 	GameDefined_0 = 28,
 	GameDefined_1,
 	GameDefined_2,
@@ -276,6 +288,7 @@ enum class EOnlineServices : uint8
 	Default = 255
 };
 
+COREONLINE_API void RegisterAliasForGameDefined(EOnlineServices Value, const TCHAR* InStr);
 COREONLINE_API const TCHAR* LexToString(EOnlineServices Value);
 COREONLINE_API void LexFromString(EOnlineServices& OutValue, const TCHAR* InStr);
 
