Change AllowedRecipients and PartyMembers in text chat analytics messages from json arrays to csv
Added LexToString to FUniqueNetIdWrapper

diff --git a/Engine/Source/Runtime/CoreOnline/Public/Online/CoreOnline.h b/Engine/Source/Runtime/CoreOnline/Public/Online/CoreOnline.h
--- a/Engine/Source/Runtime/CoreOnline/Public/Online/CoreOnline.h
+++ b/Engine/Source/Runtime/CoreOnline/Public/Online/CoreOnline.h
@@ -657,6 +657,11 @@ struct FUniqueNetIdWrapper
 		return INDEX_NONE;
 	}
 
+	friend FString LexToString(const FUniqueNetIdWrapper& Value)
+	{
+		return Value.ToString();
+	}
+
 	static FUniqueNetIdWrapper Invalid()
 	{
 		static FUniqueNetIdWrapper InvalidId(nullptr);
