Add nested variant type with JSON serialization.
#rnx

diff --git a/Engine/Plugins/Online/OnlineServices/Source/OnlineServicesInterface/Public/Online/OnlineUtils.h b/Engine/Plugins/Online/OnlineServices/Source/OnlineServicesInterface/Public/Online/OnlineUtils.h
--- a/Engine/Plugins/Online/OnlineServices/Source/OnlineServicesInterface/Public/Online/OnlineUtils.h
+++ b/Engine/Plugins/Online/OnlineServices/Source/OnlineServicesInterface/Public/Online/OnlineUtils.h
@@ -3,6 +3,7 @@
 #pragma once
 
 #include "Online/OnlineMeta.h"
+#include "Concepts/DerivedFrom.h"
 #include "Containers/Map.h"
 #include "Containers/UnrealString.h"
 #include "Containers/Utf8String.h"
@@ -14,9 +15,7 @@ struct FDateTime;
 
 namespace UE::Online {
 
-template <typename T> FString ToLogString(const TArray<T>& Array);
 template <typename T> FString ToLogString(const TSet<T>& Set);
-template <typename K, typename V> FString ToLogString(const TMap<K, V>& Map);
 template <typename T, typename U> FString ToLogString(const TPair<T, U>& Pair);
 template <typename T, ESPMode Mode> FString ToLogString(const TSharedPtr<T, Mode>& Ptr);
 template <typename T, ESPMode Mode> FString ToLogString(const TSharedRef<T, Mode>& Ref);
@@ -38,15 +37,8 @@ inline FString ToLogString(int64 Value);
 inline FString ToLogString(bool Value);
 inline FString ToLogString(float Value);
 inline FString ToLogString(double Value);
-template <typename T> std::enable_if_t<!TModels_V<Meta::COnlineMetadataAvailable, T>, FString> ToLogString(const T& Value);
-template <typename T> std::enable_if_t<TModels_V<Meta::COnlineMetadataAvailable, T>, FString> ToLogString(const T& Value);
 ONLINESERVICESINTERFACE_API FString ToLogString(const FDateTime& Time);
-
-template <typename T>
-FString ToLogString(const TArray<T>& Array)
-{
-	return FString::Printf(TEXT("[%s]"), *FString::JoinBy(Array, TEXT(", "), [](const T& Value) { return ToLogString(Value); }));
-}
+template <typename T> FString ToLogString(const T& Value);
 
 template <typename T>
 FString ToLogString(const TSet<T>& Set)
@@ -54,12 +46,6 @@ FString ToLogString(const TSet<T>& Set)
 	return FString::Printf(TEXT("{%s}"), *FString::JoinBy(Set, TEXT(", "), [](const T& Value) { return ToLogString(Value); }));
 }
 
-template <typename K, typename V>
-FString ToLogString(const TMap<K, V>& Map)
-{
-	return FString::Printf(TEXT("{%s}"), *FString::JoinBy(Map, TEXT(", "), [](const TPair<K, V>& Pair) { return ToLogString(Pair); }));
-}
-
 template <typename T, typename U>
 FString ToLogString(const TPair<T, U>& Pair)
 {
@@ -188,34 +174,60 @@ inline FString ToLogString(FPlatformUserId PlatformUserId)
 }
 
 template<typename T>
-std::enable_if_t<!TModels_V<Meta::COnlineMetadataAvailable, T>, FString> ToLogString(const T& Value)
-{
-	return LexToString(Value);
-}
+concept CDerivesFromArray =
+	requires
+	{
+		typename T::ElementType;
+	}
+	&& CDerivedFrom<T, TArray<typename T::ElementType>>;
 
 template<typename T>
-std::enable_if_t<TModels_V<Meta::COnlineMetadataAvailable, T>, FString> ToLogString(const T& Value)
-{
-	FString LogString;
-	LogString += TEXT("{ ");
-	bool bFirst = true;
-	Meta::VisitFields(Value, [&LogString, &bFirst](const TCHAR* Name, auto& Field)
-		{
-			if (bFirst)
-			{
-				bFirst = false;
-			}
-			else
+concept CDerivesFromMap =
+	requires
+	{
+		typename T::KeyType;
+		typename T::ValueType;
+	}
+	&& CDerivedFrom<T, TMap<typename T::KeyType, typename T::ValueType>>;
+
+template <typename T>
+FString ToLogString(const T& Value)
+{
+	if constexpr (CDerivesFromArray<T>)
+	{
+		return FString::Printf(TEXT("[%s]"), *FString::JoinBy(Value, TEXT(", "), [](const typename T::ElementType& Element) { return ToLogString(Element); }));
+	}
+	else if constexpr (CDerivesFromMap<T>)
+	{
+		return FString::Printf(TEXT("{%s}"), *FString::JoinBy(Value, TEXT(", "), [](const TPair<typename T::KeyType, typename T::ValueType>& Pair) { return ToLogString(Pair); }));
+	}
+	else if constexpr (TModels_V<Meta::COnlineMetadataAvailable, T>)
+	{
+		FString LogString;
+		LogString += TEXT("{ ");
+		bool bFirst = true;
+		Meta::VisitFields(Value, [&LogString, &bFirst](const TCHAR* Name, auto& Field)
 			{
-				LogString += TEXT(", ");
-			}
-			LogString += Name;
-			LogString += TEXT(": ");
-			LogString += ToLogString(Field);
-		});
-	LogString += TEXT(" }");
-
-	return LogString;
+				if (bFirst)
+				{
+					bFirst = false;
+				}
+				else
+				{
+					LogString += TEXT(", ");
+				}
+				LogString += Name;
+				LogString += TEXT(": ");
+				LogString += ToLogString(Field);
+			});
+		LogString += TEXT(" }");
+
+		return LogString;
+	}
+	else
+	{
+		return LexToString(Value);
+	}
 }
 
 /* UE::Online */ }
\ No newline at end of file
diff --git a/Engine/Source/Runtime/CoreOnline/CoreOnline.Build.cs b/Engine/Source/Runtime/CoreOnline/CoreOnline.Build.cs
--- a/Engine/Source/Runtime/CoreOnline/CoreOnline.Build.cs
+++ b/Engine/Source/Runtime/CoreOnline/CoreOnline.Build.cs
@@ -10,7 +10,8 @@ public class CoreOnline : ModuleRules
 		PublicDependencyModuleNames.AddRange(
 			new string[] {
 				"Core",
-				"CoreUObject"
+				"CoreUObject",
+				"Json"
 			}
 		);
 
diff --git a/Engine/Source/Runtime/CoreOnline/Private/Online/NestedVariantJson.cpp b/Engine/Source/Runtime/CoreOnline/Private/Online/NestedVariantJson.cpp
new file mode 100644
--- /dev/null
+++ b/Engine/Source/Runtime/CoreOnline/Private/Online/NestedVariantJson.cpp
@@ -0,0 +1,186 @@
+// Copyright Epic Games, Inc. All Rights Reserved.
+
+#include "Online/NestedVariantJson.h"
+
+#include "Dom/JsonObject.h"
+#include "Policies/CondensedJsonPrintPolicy.h"
+#include "Serialization/JsonSerializer.h"
+
+// TODO: Update the JSON module to handle nested variants of any type directly so this file is not needed.
+
+FString NestedVariantToJson(const FNestedVariantJson::FMapPtr& Map)
+{
+	FString Result;
+	if (!Map)
+	{
+		return Result;
+	}
+
+	using FJsonWriterFactory = TJsonWriterFactory<TCHAR, TCondensedJsonPrintPolicy<TCHAR>>;
+	using FJsonWriter = TJsonWriter<TCHAR, TCondensedJsonPrintPolicy<TCHAR>>;
+
+	const FNestedVariantJson::FMapRef MapRef = Map.ToSharedRef();
+	TSharedRef<FJsonObject> JsonObject = NestedVariantToJsonObject(MapRef);
+	TSharedRef<FJsonWriter> Writer = FJsonWriterFactory::Create(&Result);
+	verify(FJsonSerializer::Serialize(JsonObject, Writer));
+	return Result;
+}
+
+TSharedRef<FJsonObject> NestedVariantToJsonObject(const FNestedVariantJson::FMapRef& Map)
+{
+	TSharedRef<FJsonObject> JsonObject = MakeShared<FJsonObject>();
+	JsonObject->Values.Reserve(Map->Num());
+
+	for (const TPair<FString, FNestedVariantJson::FValue>& Pair : *Map)
+	{
+		JsonObject->SetField(Pair.Key, NestedVariantToJsonValue(Pair.Value));
+	}
+
+	return JsonObject;
+}
+
+TSharedPtr<FJsonValue> NestedVariantToJsonValue(const FNestedVariantJson::FValue& Value)
+{
+	switch (Value.GetIndex())
+	{
+	case FNestedVariantJson::FValue::IndexOfType<bool>():
+		return MakeShared<FJsonValueBoolean>(Value.Get<bool>());
+	case FNestedVariantJson::FValue::IndexOfType<int64>():
+		return MakeShared<FJsonValueNumber>(Value.Get<int64>());
+	case FNestedVariantJson::FValue::IndexOfType<double>():
+		return MakeShared<FJsonValueNumber>(Value.Get<double>());
+	case FNestedVariantJson::FValue::IndexOfType<FString>():
+		return MakeShared<TJsonValueString<TCHAR>>(Value.Get<FString>());
+	case FNestedVariantJson::FValue::IndexOfType<FNestedVariantJson::FArrayRef>():
+	{
+		const FNestedVariantJson::FArrayRef& Array = Value.Get<FNestedVariantJson::FArrayRef>();
+		TArray<TSharedPtr<FJsonValue>> JsonArray;
+		JsonArray.Reserve(Array->Num());
+
+		for (const FNestedVariantJson::FValue& ArrayValue : *Array)
+		{
+			JsonArray.Add(NestedVariantToJsonValue(ArrayValue));
+		}
+
+		return MakeShared<FJsonValueArray>(MoveTemp(JsonArray));
+	}
+	case FNestedVariantJson::FValue::IndexOfType<FNestedVariantJson::FMapRef>():
+		return MakeShared<FJsonValueObject>(NestedVariantToJsonObject(Value.Get<FNestedVariantJson::FMapRef>()));
+	default:
+		checkNoEntry();
+	}
+
+	return nullptr;
+}
+
+void NestedVariantFromJson(const char* InJson, FNestedVariantJson::FMapRef& Map)
+{
+	if (!InJson || strlen(InJson) == 0)
+	{
+		return;
+	}
+
+	TSharedPtr<FJsonObject> JsonObject = MakeShared<FJsonObject>();
+	TSharedRef<TJsonReader<UTF8CHAR>> Reader = TJsonReaderFactory<UTF8CHAR>::Create(InJson);
+	verify(FJsonSerializer::Deserialize(Reader, JsonObject));
+	if (!JsonObject)
+	{
+		return;
+	}
+
+	NestedVariantFromJsonObject(JsonObject, Map);
+}
+
+void NestedVariantFromJsonObject(const TSharedPtr<FJsonObject>& JsonObject, FNestedVariantJson::FMapRef& Map)
+{
+	for (const TPair<FString, TSharedPtr<FJsonValue>>& Pair : JsonObject->Values)
+	{
+		switch (Pair.Value->Type)
+		{
+		case EJson::String:
+			Map->AddVariant(Pair.Key, Pair.Value->AsString());
+			break;
+		case EJson::Number:
+		{
+			int64 IntValue;
+			if (Pair.Value->TryGetNumber(IntValue))
+			{
+				Map->AddVariant(Pair.Key, IntValue);
+			}
+			else
+			{
+				double DoubleValue = 0;
+				verify(Pair.Value->TryGetNumber(DoubleValue));
+				Map->AddVariant(Pair.Key, DoubleValue);
+			}
+			break;
+		}
+		case EJson::Boolean:
+			Map->AddVariant(Pair.Key, Pair.Value->AsBool());
+			break;
+		case EJson::Array:
+		{
+			FNestedVariantJson::FArrayRef NestedArray = FNestedVariantJson::FArray::CreateVariant();
+			NestedVariantFromJsonArray(Pair.Value->AsArray(), NestedArray);
+			Map->AddVariant(Pair.Key, NestedArray);
+			break;
+		}
+		case EJson::Object:
+		{
+			FNestedVariantJson::FMapRef NestedMap = FNestedVariantJson::FMap::CreateVariant();
+			NestedVariantFromJsonObject(Pair.Value->AsObject(), NestedMap);
+			Map->AddVariant(Pair.Key, NestedMap);
+			break;
+		}
+		default:
+			checkNoEntry();
+		}
+	}
+}
+
+void NestedVariantFromJsonArray(const TArray<TSharedPtr<FJsonValue>>& JsonArray, FNestedVariantJson::FArrayRef& Array)
+{
+	for (const TSharedPtr<FJsonValue>& Value : JsonArray)
+	{
+		switch (Value->Type)
+		{
+		case EJson::String:
+			Array->AddVariant(Value->AsString());
+			break;
+		case EJson::Number:
+		{
+			int64 IntValue;
+			if (Value->TryGetNumber(IntValue))
+			{
+				Array->AddVariant(IntValue);
+			}
+			else
+			{
+				double DoubleValue = 0;
+				verify(Value->TryGetNumber(DoubleValue));
+				Array->AddVariant(DoubleValue);
+			}
+			break;
+		}
+		case EJson::Boolean:
+			Array->AddVariant(Value->AsBool());
+			break;
+		case EJson::Array:
+		{
+			FNestedVariantJson::FArrayRef NestedArray = FNestedVariantJson::FArray::CreateVariant();
+			NestedVariantFromJsonArray(Value->AsArray(), NestedArray);
+			Array->AddVariant(NestedArray);
+			break;
+		}
+		case EJson::Object:
+		{
+			FNestedVariantJson::FMapRef NestedMap = FNestedVariantJson::FMap::CreateVariant();
+			NestedVariantFromJsonObject(Value->AsObject(), NestedMap);
+			Array->AddVariant(NestedMap);
+			break;
+		}
+		default:
+			checkNoEntry();
+		}
+	}
+}
diff --git a/Engine/Source/Runtime/CoreOnline/Public/Online/NestedVariant.h b/Engine/Source/Runtime/CoreOnline/Public/Online/NestedVariant.h
new file mode 100644
--- /dev/null
+++ b/Engine/Source/Runtime/CoreOnline/Public/Online/NestedVariant.h
@@ -0,0 +1,249 @@
+// Copyright Epic Games, Inc. All Rights Reserved.
+
+#pragma once
+
+#include "Containers/Array.h"
+#include "Containers/Map.h"
+#include "Misc/TVariant.h"
+#include "Templates/SharedPointer.h"
+
+/**
+ * This set of classes enables building nested TVariant structures using shared pointers to TArrays and TMaps. This can be extended
+ * to other core container types or custom types as needed. If this proves useful beyond online use cases, this could move to Core.
+ * 
+ * This allows for usage such as:
+ *
+ * using FNestedVariant = TNestedVariant<FString, bool, int64, double, FString>; // The first FString is the TMap key type.
+ * FNestedVariant::FMapRef Data = FNestedVariant::FMap::CreateVariant();
+ *
+ * Data->AddVariant(TEXT("MyString"), FString(TEXT("1234-5678-9000")));
+ * Data->AddVariant(TEXT("MyInt"), 42LL);
+ * Data->AddVariant(TEXT("MyObject"), FNestedVariant::FMap::CreateVariant())->AddVariant(TEXT("MyNestedObjectDouble"), 1.0);
+ * FNestedVariant::FArrayRef& Array = Data->AddVariant(TEXT("MyArray"), FNestedVariant::FArray::CreateVariant());
+ *
+ * Array->AddVariant(FNestedVariant::FMap::CreateVariant())->AddVariant(TEXT("MyNestedArrayObjectInt"), 100LL);
+ * Array->AddVariant(FNestedVariant::FArray::CreateVariant())->AddVariant(1.23);
+ * Array->AddVariant(4.0);
+ * Array->AddVariant(true);
+ *
+ * FNestedVariant::FMapPtr NestedMap;
+ * if (Array->GetVariant(0, NestedMap))
+ * {
+ *     // Found nested map in array.
+ * }
+ *
+ * double ArrayDouble;
+ * if (Array->GetVariant(2, ArrayDouble)) { ... }
+ *
+ * int64 MyInt;
+ * if (Data->GetVariant(TEXT("MyInt"), MyInt)) { ... }
+ *
+ * FNestedVariant::FArrayPtr NestedArray;
+ * if (Data->GetVariant(TEXT("MyArray"), NestedArray))
+ * {
+ *     bool Value;
+ *     if (NestedArray->GetVariant(3, Value)) { ... }
+ * }
+ * 
+ * OnlineUtils.h has ToLogString() overrides to allow for: ToLogString(Data);
+ * Which would output: {MyString:1234-5678-9000, MyInt:42, MyObject:{MyNestedObjectDouble:1.00}, MyArray:[{MyNestedArrayObjectInt:100}, [1.23], 4.00, true]}
+ */
+
+
+/** TArray that holds a TVariant with helper functions for working with variant types. */
+template<typename... ValueTypes>
+class TVariantArray : public TArray<TVariant<ValueTypes...>>, public TSharedFromThis<TVariantArray<ValueTypes...>>
+{
+public:
+	using VariantType = TVariant<ValueTypes...>;
+	using ArrayType = TArray<VariantType>;
+
+	static TSharedRef<TVariantArray<ValueTypes...>> CreateVariant()
+	{
+		return MakeShared<TVariantArray<ValueTypes...>>();
+	}
+
+	template<typename ValueType, typename VariantValueType = std::remove_cv_t<std::remove_reference_t<ValueType>>>
+	VariantValueType& AddVariant(ValueType&& Value)
+	{
+		return ArrayType::Emplace_GetRef(VariantType(TInPlaceType<VariantValueType>(), Forward<ValueType>(Value))).template Get<VariantValueType>();
+	}
+
+	template<typename ValueType>
+	void AddVariantArray(const TArray<ValueType>& Array)
+	{
+		for (const ValueType& Value : Array)
+		{
+			AddVariant(Value);
+		}
+	}
+
+	template<typename ValueType>
+	bool GetVariant(ArrayType::SizeType Index, TSharedPtr<ValueType>& OutValue) const
+	{
+		if (Index >= 0 && Index < ArrayType::Num())
+		{
+			if (ArrayType::operator[](Index).template IsType<TSharedRef<ValueType>>())
+			{
+				OutValue = ArrayType::operator[](Index).template Get<TSharedRef<ValueType>>();
+				return true;
+			}
+		}
+
+		return false;
+	}
+
+	template<typename ValueType>
+	bool GetVariant(ArrayType::SizeType Index, ValueType& OutValue) const
+	{
+		if (Index >= 0 && Index < ArrayType::Num())
+		{
+			if (ArrayType::operator[](Index).template IsType<ValueType>())
+			{
+				OutValue = ArrayType::operator[](Index).template Get<ValueType>();
+				return true;
+			}
+		}
+
+		return false;
+	}
+
+	template<typename ValueType>
+	bool GetVariantArray(TArray<ValueType>& OutValue) const
+	{
+		OutValue.Reserve(ArrayType::Num());
+		for (int32 Index = 0; Index < ArrayType::Num(); Index++)
+		{
+			if (ArrayType::operator[](Index).template IsType<ValueType>())
+			{
+				OutValue.Add(ArrayType::operator[](Index).template Get<ValueType>());
+			}
+		}
+
+		return ArrayType::Num() == OutValue.Num();
+	}
+};
+
+/** TMap that holds a TVariant with helper functions for working with variant types. */
+template<typename KeyType, typename... ValueTypes>
+class TVariantMap : public TMap<KeyType, TVariant<ValueTypes...>>, public TSharedFromThis<TVariantMap<KeyType, ValueTypes...>>
+{
+public:
+	using VariantType = TVariant<ValueTypes...>;
+	using MapType = TMap<KeyType, VariantType>;
+
+	static TSharedRef<TVariantMap<KeyType, ValueTypes...>> CreateVariant()
+	{
+		return MakeShared<TVariantMap<KeyType, ValueTypes...>>();
+	}
+
+	template<typename ValueType, typename AddKeyType, typename VariantValueType = std::remove_cv_t<std::remove_reference_t<ValueType>>>
+	VariantValueType& AddVariant(AddKeyType&& Key, ValueType&& Value)
+	{
+		return MapType::Emplace(Forward<AddKeyType>(Key), VariantType(TInPlaceType<VariantValueType>(), Forward<ValueType>(Value))).template Get<VariantValueType>();
+	}
+
+	template<typename ValueType>
+	void AddVariantMap(const TMap<KeyType, ValueType>& Map)
+	{
+		for (const TPair<KeyType, ValueType>& Pair : Map)
+		{
+			AddVariant(Pair.Key, Pair.Value);
+		}
+	}
+
+	template<typename ValueType, typename GetKeyType>
+	bool GetVariant(const GetKeyType& Key, TSharedPtr<ValueType>& OutValue) const
+	{
+		if (const VariantType* Attribute = MapType::Find(Key))
+		{
+			if (Attribute->template IsType<TSharedRef<ValueType>>())
+			{
+				OutValue = Attribute->template Get<TSharedRef<ValueType>>();
+				return true;
+			}
+		}
+
+		return false;
+	}
+
+	template<typename ValueType, typename GetKeyType>
+	bool GetVariant(const GetKeyType& Key, ValueType& OutValue) const
+	{
+		if (const VariantType* Attribute = MapType::Find(Key))
+		{
+			if (Attribute->template IsType<ValueType>())
+			{
+				OutValue = Attribute->template Get<ValueType>();
+				return true;
+			}
+		}
+
+		return false;
+	}
+
+	template<typename ValueType>
+	bool GetVariantMap(TMap<KeyType, ValueType>& OutValue) const
+	{
+		OutValue.Reserve(MapType::Num());
+		for (const TPair<KeyType, VariantType>& Pair : MapType::Pairs)
+		{
+			if (Pair.Value.template IsType<ValueType>())
+			{
+				OutValue.Add(Pair.Key, Pair.Value.template Get<ValueType>());
+			}
+		}
+
+		return MapType::Num() == OutValue.Num();
+	}
+};
+
+// Classes for working with nested TVariants that also have TArray and TMap references holding TVariants of the same types.
+// This is just one possible nesting structure. Other custom structures can be built for various use cases (no arrays, arrays with fixed types, etc).
+// First declare types to allow for nesting.
+
+template<typename KeyType, typename... ValueTypes> class TNestedVariantArray;
+template<typename KeyType, typename... ValueTypes> using TNestedVariantArrayRef = TSharedRef<TNestedVariantArray<KeyType, ValueTypes...>>;
+template<typename KeyType, typename... ValueTypes> using TNestedVariantArrayPtr = TSharedPtr<TNestedVariantArray<KeyType, ValueTypes...>>;
+
+template<typename KeyType, typename... ValueTypes> class TNestedVariantMap;
+template<typename KeyType, typename... ValueTypes> using TNestedVariantMapRef = TSharedRef<TNestedVariantMap<KeyType, ValueTypes...>>;
+template<typename KeyType, typename... ValueTypes> using TNestedVariantMapPtr = TSharedPtr<TNestedVariantMap<KeyType, ValueTypes...>>;
+
+template<typename KeyType, typename... ValueTypes> using FNestedVariantValue = TVariant<ValueTypes..., TNestedVariantArrayRef<KeyType, ValueTypes...>, TNestedVariantMapRef<KeyType, ValueTypes...>>;
+
+/** Nested TVariantArray that can hold all provided types plus nested array and map references. */
+template<typename KeyType, typename... ValueTypes>
+class TNestedVariantArray : public TVariantArray<ValueTypes..., TNestedVariantArrayRef<KeyType, ValueTypes...>, TNestedVariantMapRef<KeyType, ValueTypes...>>
+{
+public:
+	static TNestedVariantArrayRef<KeyType, ValueTypes...> CreateVariant() {
+		return MakeShared<TNestedVariantArray<KeyType, ValueTypes...>>();
+	}
+};
+
+/** Nested TVariantMap that can hold all provided types plus nested array and map references. */
+template<typename KeyType, typename... ValueTypes>
+class TNestedVariantMap : public TVariantMap<KeyType, ValueTypes..., TNestedVariantArrayRef<KeyType, ValueTypes...>, TNestedVariantMapRef<KeyType, ValueTypes...>>
+{
+public:
+	static TNestedVariantMapRef<KeyType, ValueTypes...> CreateVariant() {
+		return MakeShared<TNestedVariantMap<KeyType, ValueTypes...>>();
+	}
+};
+
+/** Convenience class to declare all nested types at once. */
+template<typename KeyType, typename... ValueTypes>
+class TNestedVariant
+{
+public:
+	using FArray = TNestedVariantArray<KeyType, ValueTypes...>;
+	using FArrayRef = TNestedVariantArrayRef<KeyType, ValueTypes...>;
+	using FArrayPtr = TNestedVariantArrayPtr<KeyType, ValueTypes...>;
+
+	using FMap = TNestedVariantMap<KeyType, ValueTypes...>;
+	using FMapRef = TNestedVariantMapRef<KeyType, ValueTypes...>;
+	using FMapPtr = TNestedVariantMapPtr<KeyType, ValueTypes...>;
+
+	using FValue = FNestedVariantValue<KeyType, ValueTypes...>;
+};
diff --git a/Engine/Source/Runtime/CoreOnline/Public/Online/NestedVariantJson.h b/Engine/Source/Runtime/CoreOnline/Public/Online/NestedVariantJson.h
new file mode 100644
--- /dev/null
+++ b/Engine/Source/Runtime/CoreOnline/Public/Online/NestedVariantJson.h
@@ -0,0 +1,20 @@
+// Copyright Epic Games, Inc. All Rights Reserved.
+
+#pragma once
+
+#include "Containers/UnrealString.h"
+#include "Online/NestedVariant.h"
+
+// TODO: Update the JSON module to handle nested variants of any type directly so this file is not needed.
+
+using FNestedVariantJson = TNestedVariant<FString, bool, int64, double, FString>;
+class FJsonObject;
+class FJsonValue;
+
+COREONLINE_API FString NestedVariantToJson(const FNestedVariantJson::FMapPtr& Map);
+COREONLINE_API TSharedRef<FJsonObject> NestedVariantToJsonObject(const FNestedVariantJson::FMapRef& Map);
+COREONLINE_API TSharedPtr<FJsonValue> NestedVariantToJsonValue(const FNestedVariantJson::FValue& Value);
+
+COREONLINE_API void NestedVariantFromJson(const char* InJson, FNestedVariantJson::FMapRef& Map);
+COREONLINE_API void NestedVariantFromJsonObject(const TSharedPtr<FJsonObject>& JsonObject, FNestedVariantJson::FMapRef& Map);
+COREONLINE_API void NestedVariantFromJsonArray(const TArray<TSharedPtr<FJsonValue>>& JsonArray, FNestedVariantJson::FArrayRef& Array);
