- Add config switch for turning on epic modular v2 plugins

#rnx
#DisabledValidations LowLevelTestsOnlineTests, LowLevelTestsFNOnlineFrameworkTests, lowLevelTestsOSMCPTests

diff --git a/Engine/Plugins/Online/OnlineServices/Source/OnlineServicesInterface/Private/Online/OnlineServicesRegistry.cpp b/Engine/Plugins/Online/OnlineServices/Source/OnlineServicesInterface/Private/Online/OnlineServicesRegistry.cpp
--- a/Engine/Plugins/Online/OnlineServices/Source/OnlineServicesInterface/Private/Online/OnlineServicesRegistry.cpp
+++ b/Engine/Plugins/Online/OnlineServices/Source/OnlineServicesInterface/Private/Online/OnlineServicesRegistry.cpp
@@ -181,6 +181,10 @@ TSharedPtr<IOnlineServices> FOnlineServicesRegistry::CreateServices(EOnlineServi
 {
 	OnlineServices = ResolveServiceName(OnlineServices);
 
+	// Ensure if modular plugins are enabled (EpicAccount + EpicGame) that nothing asks for Epic.
+	const bool bIsMonolithicPlugin = OnlineServices == EOnlineServices::Epic;
+	ensure(!(TEMP_ShouldUseEpicModularPlugins() && bIsMonolithicPlugin));
+
 	TSharedPtr<IOnlineServices> Services;
 
 	if (FFactoryAndPriority* FactoryAndPriority = ServicesFactories.Find(OnlineServices))
diff --git a/Engine/Plugins/Online/OnlineServicesEOS/Source/Private/Online/OnlineServicesEOSModule.cpp b/Engine/Plugins/Online/OnlineServicesEOS/Source/Private/Online/OnlineServicesEOSModule.cpp
--- a/Engine/Plugins/Online/OnlineServicesEOS/Source/Private/Online/OnlineServicesEOSModule.cpp
+++ b/Engine/Plugins/Online/OnlineServicesEOS/Source/Private/Online/OnlineServicesEOSModule.cpp
@@ -37,10 +37,9 @@ void FOnlineServicesEOSModule::StartupModule()
 {
 	FModuleManager::Get().LoadModuleChecked(TEXT("OnlineServicesInterface"));
 
-	const bool bEnableEOSSplitPlugins = FParse::Param(FCommandLine::Get(), TEXT("EnableEOSSplitPlugins"));
-	if (bEnableEOSSplitPlugins)
+	if (TEMP_ShouldUseEpicModularPlugins())
 	{
-		UE_LOG(LogOnlineServices, Verbose, TEXT("EnableEOSSplitPlugins launch parameter detected. Skipping FOnlineServicesEOSModule startup."));
+		UE_LOG(LogOnlineServices, Verbose, TEXT("TEMP_ShouldUseEpicModularPlugins() = true, skipping FOnlineServicesEOSModule startup."));
 		return;
 	}
 
diff --git a/Engine/Plugins/Online/OnlineServicesEOSGS/Source/Private/Online/OnlineServicesEOSGSModule.cpp b/Engine/Plugins/Online/OnlineServicesEOSGS/Source/Private/Online/OnlineServicesEOSGSModule.cpp
--- a/Engine/Plugins/Online/OnlineServicesEOSGS/Source/Private/Online/OnlineServicesEOSGSModule.cpp
+++ b/Engine/Plugins/Online/OnlineServicesEOSGS/Source/Private/Online/OnlineServicesEOSGSModule.cpp
@@ -36,10 +36,9 @@ void FOnlineServicesEOSGSModule::StartupModule()
 {
 	FModuleManager::Get().LoadModuleChecked(TEXT("OnlineServicesInterface"));
 
-	const bool bEnableEOSSplitPlugins = FParse::Param(FCommandLine::Get(), TEXT("EnableEOSSplitPlugins"));
-	if (bEnableEOSSplitPlugins)
+	if (TEMP_ShouldUseEpicModularPlugins())
 	{
-		UE_LOG(LogOnlineServices, Verbose, TEXT("EnableEOSSplitPlugins launch parameter detected. Skipping FOnlineServicesEOSGSModule startup."));
+		UE_LOG(LogOnlineServices, Verbose, TEXT("TEMP_ShouldUseEpicModularPlugins() = true, skipping FOnlineServicesEOSGSModule startup."));
 		return;
 	}
 
diff --git a/Engine/Source/Runtime/CoreOnline/Private/Online/CoreOnline.cpp b/Engine/Source/Runtime/CoreOnline/Private/Online/CoreOnline.cpp
--- a/Engine/Source/Runtime/CoreOnline/Private/Online/CoreOnline.cpp
+++ b/Engine/Source/Runtime/CoreOnline/Private/Online/CoreOnline.cpp
@@ -2,6 +2,9 @@
 
 #include "Online/CoreOnline.h"
 
+#include "CoreGlobals.h"
+#include "Misc/CommandLine.h"
+#include "Misc/ConfigCacheIni.h"
 #include "Misc/LazySingleton.h"
 #include "Online/CoreOnlinePrivate.h"
 
@@ -236,6 +239,18 @@ void LexFromString(EOnlineServices& OutValue, const TCHAR* InStr)
 	}
 }
 
+bool TEMP_ShouldUseEpicModularPlugins()
+{
+	if (FParse::Param(FCommandLine::Get(), TEXT("EnableEOSSplitPlugins")))
+	{
+		return true;
+	}
+
+	bool bShouldUseEpicModularPlugins = false;
+	GConfig->GetBool(TEXT("CoreOnline"), TEXT("bShouldUseEpicModularPlugins"), bShouldUseEpicModularPlugins, GEngineIni);
+	return bShouldUseEpicModularPlugins;
+}
+
 FOnlineIdRegistryRegistry& FOnlineIdRegistryRegistry::Get()
 {
 	return TLazySingleton<FOnlineIdRegistryRegistry>::Get();
diff --git a/Engine/Source/Runtime/CoreOnline/Public/Online/CoreOnline.h b/Engine/Source/Runtime/CoreOnline/Public/Online/CoreOnline.h
--- a/Engine/Source/Runtime/CoreOnline/Public/Online/CoreOnline.h
+++ b/Engine/Source/Runtime/CoreOnline/Public/Online/CoreOnline.h
@@ -292,6 +292,13 @@ COREONLINE_API void RegisterAliasForGameDefined(EOnlineServices Value, const TCH
 COREONLINE_API const TCHAR* LexToString(EOnlineServices Value);
 COREONLINE_API void LexFromString(EOnlineServices& OutValue, const TCHAR* InStr);
 
+/**
+ * Temporary
+ * Signals that the EOnlineServices::Epic monolithic plugins should not be used, and EOnlineServices::EpicAccount/EpicGame should be used instead.
+ * This is a stepping stone to promoting the EOnlineServices::EpicAccount/EpicGame modular plugins and deprecating the monolithic ones.
+ */
+COREONLINE_API bool TEMP_ShouldUseEpicModularPlugins();
+
 /**
  * A handle to an id which uniquely identifies a persistent or transient online resource, i.e. account/session/party etc, within a given Online Services provider.
  * At most one id, and therefore one handle, exists for any given resource. The id and handle persist until the OnlineServices module is unloaded.
