Added missing value assignation

diff --git a/Engine/Source/Runtime/CoreOnline/Private/Online/NestedVariantJson.cpp b/Engine/Source/Runtime/CoreOnline/Private/Online/NestedVariantJson.cpp
--- a/Engine/Source/Runtime/CoreOnline/Private/Online/NestedVariantJson.cpp
+++ b/Engine/Source/Runtime/CoreOnline/Private/Online/NestedVariantJson.cpp
@@ -27,7 +27,7 @@ FString NestedVariantToJson(const FNestedVariantJson::FMapPtr& Map)
 
 TSharedRef<FJsonObject> NestedVariantToJsonObject(const FNestedVariantJson::FMapRef& Map)
 {
-	FString JsonStr = NestedVariantToJson(Map);
+	const FString& JsonStr = NestedVariantToJson(Map);
 	TSharedRef<TJsonReader<TCHAR>> Reader = TJsonReaderFactory<TCHAR>::Create(JsonStr);	
 	TSharedPtr<FJsonObject> JsonObject = MakeShared<FJsonObject>();
 	verify(FJsonSerializer::Deserialize(Reader, JsonObject));
@@ -43,8 +43,10 @@ void NestedVariantFromJson(const char* InJson, FNestedVariantJson::FMapRef& Map)
 	}
 
 	TSharedRef<TJsonReader<TCHAR>> Reader = TJsonReaderFactory<TCHAR>::Create(UTF8_TO_TCHAR(InJson));
-	FNestedVariantJson::FMapPtr MapPtr =  Map.ToSharedPtr();
+	FNestedVariantJson::FMapPtr MapPtr;
 	verify(FJsonSerializerNestedVariant::Deserialize(Reader, MapPtr));
+
+	Map->Append(*MapPtr);
 }
 
 void NestedVariantFromJsonObject(const TSharedPtr<FJsonObject>& JsonObject, FNestedVariantJson::FMapRef& Map)
