Added check for shared ptr for extra safety

diff --git a/Engine/Source/Runtime/CoreOnline/Private/Online/NestedVariantJson.cpp b/Engine/Source/Runtime/CoreOnline/Private/Online/NestedVariantJson.cpp
--- a/Engine/Source/Runtime/CoreOnline/Private/Online/NestedVariantJson.cpp
+++ b/Engine/Source/Runtime/CoreOnline/Private/Online/NestedVariantJson.cpp
@@ -46,7 +46,10 @@ void NestedVariantFromJson(const char* InJson, FNestedVariantJson::FMapRef& Map)
 	FNestedVariantJson::FMapPtr MapPtr;
 	verify(FJsonSerializerNestedVariant::Deserialize(Reader, MapPtr));
 
-	Map->Append(*MapPtr);
+	if (MapPtr.IsValid())
+	{
+		Map->Append(*MapPtr);
+	}
 }
 
 void NestedVariantFromJsonObject(const TSharedPtr<FJsonObject>& JsonObject, FNestedVariantJson::FMapRef& Map)
