Added MoveTemp to map assignation to avoid copy

diff --git a/Engine/Source/Runtime/CoreOnline/Private/Online/NestedVariantJson.cpp b/Engine/Source/Runtime/CoreOnline/Private/Online/NestedVariantJson.cpp
--- a/Engine/Source/Runtime/CoreOnline/Private/Online/NestedVariantJson.cpp
+++ b/Engine/Source/Runtime/CoreOnline/Private/Online/NestedVariantJson.cpp
@@ -48,7 +48,7 @@ void NestedVariantFromJson(const char* InJson, FNestedVariantJson::FMapRef& Map)
 
 	if (MapPtr.IsValid())
 	{
-		Map->Append(*MapPtr);
+		Map->Append(MoveTemp(*MapPtr));
 	}
 }
 
