Adding direct map support for nestedVariant serialization

diff --git a/Engine/Source/Runtime/CoreOnline/Private/Online/NestedVariantJson.cpp b/Engine/Source/Runtime/CoreOnline/Private/Online/NestedVariantJson.cpp
--- a/Engine/Source/Runtime/CoreOnline/Private/Online/NestedVariantJson.cpp
+++ b/Engine/Source/Runtime/CoreOnline/Private/Online/NestedVariantJson.cpp
@@ -35,6 +35,12 @@ TSharedRef<FJsonObject> NestedVariantToJsonObject(const FNestedVariantJson::FMap
 	return JsonObject.ToSharedRef();
 }
 
+TSharedRef<FJsonObject> NestedVariantToJsonObject(const FNestedVariantJson::FMap& Map)
+{
+	const FNestedVariantJson::FMapRef Ref = MakeShared<FNestedVariantJson::FMap>(Map);
+	return NestedVariantToJsonObject(Ref);
+}
+
 void NestedVariantFromJson(const char* InJson, FNestedVariantJson::FMapRef& Map)
 {
 	if (!InJson || strlen(InJson) == 0)
@@ -52,6 +58,13 @@ void NestedVariantFromJson(const char* InJson, FNestedVariantJson::FMapRef& Map)
 	}
 }
 
+void NestedVariantFromJson(const char* InJson, FNestedVariantJson::FMap& Map)
+{
+	FNestedVariantJson::FMapRef Ref = MakeShared<FNestedVariantJson::FMap>();
+	NestedVariantFromJson(InJson, Ref);
+	Map = MoveTemp(*Ref);
+}
+
 void NestedVariantFromJsonObject(const TSharedPtr<FJsonObject>& JsonObject, FNestedVariantJson::FMapRef& Map)
 {
 	FString JsonStr;
@@ -61,6 +74,12 @@ void NestedVariantFromJsonObject(const TSharedPtr<FJsonObject>& JsonObject, FNes
 	NestedVariantFromJson(TCHAR_TO_UTF8(*JsonStr), Map);
 }
 
+void NestedVariantFromJsonObject(const TSharedPtr<FJsonObject>& JsonObject, FNestedVariantJson::FMap& Map)
+{
+	FNestedVariantJson::FMapRef Ref = MakeShared<FNestedVariantJson::FMap>(Map);
+	return NestedVariantFromJsonObject(JsonObject, Ref);
+}
+
 void LexFromString(FNestedVariantJson::FMap& OutValue, const TCHAR* InString)
 {
 	FNestedVariantJson::FMapRef  OutValueRef = FNestedVariantJson::FMap::CreateVariant();
diff --git a/Engine/Source/Runtime/CoreOnline/Public/Online/NestedVariantJson.h b/Engine/Source/Runtime/CoreOnline/Public/Online/NestedVariantJson.h
--- a/Engine/Source/Runtime/CoreOnline/Public/Online/NestedVariantJson.h
+++ b/Engine/Source/Runtime/CoreOnline/Public/Online/NestedVariantJson.h
@@ -16,9 +16,12 @@ class FJsonValue;
 
 COREONLINE_API FString NestedVariantToJson(const FNestedVariantJson::FMapPtr& Map);
 COREONLINE_API TSharedRef<FJsonObject> NestedVariantToJsonObject(const FNestedVariantJson::FMapRef& Map);
+COREONLINE_API TSharedRef<FJsonObject> NestedVariantToJsonObject(const FNestedVariantJson::FMap& Map);
 
 COREONLINE_API void NestedVariantFromJson(const char* InJson, FNestedVariantJson::FMapRef& Map);
+COREONLINE_API void NestedVariantFromJson(const char* InJson, FNestedVariantJson::FMap& Map);
 COREONLINE_API void NestedVariantFromJsonObject(const TSharedPtr<FJsonObject>& JsonObject, FNestedVariantJson::FMapRef& Map);
+COREONLINE_API void NestedVariantFromJsonObject(const TSharedPtr<FJsonObject>& JsonObject, FNestedVariantJson::FMap& Map);
 
 COREONLINE_API void LexFromString(FNestedVariantJson::FMap& OutValue, const TCHAR* InString);
 
